Mesh Editor: New "EditableMesh" plugin (replaces "MeshEditingRuntime" module)
- "MeshEditingRuntime" module has been renamed/moved to an "EditableMesh" module inside the new "EditableMesh" plugin
- The existing "MeshEditor" plugin has a dependency on the new "EditableMesh" plugin
- No functional changes

diff --git a/Engine/Plugins/Editor/MeshEditor/MeshEditor.uplugin b/Engine/Plugins/Editor/MeshEditor/MeshEditor.uplugin
--- a/Engine/Plugins/Editor/MeshEditor/MeshEditor.uplugin
+++ b/Engine/Plugins/Editor/MeshEditor/MeshEditor.uplugin
@@ -26,5 +26,11 @@
 			"Type" : "Editor",
 			"LoadingPhase" : "Default"
 		}
+	],
+	"Plugins": [
+		{
+			"Name": "EditableMesh",
+			"Enabled": true
+		}
 	]
 }
\ No newline at end of file
diff --git a/Engine/Plugins/Editor/MeshEditor/Source/MeshEditor/MeshEditor.Build.cs b/Engine/Plugins/Editor/MeshEditor/Source/MeshEditor/MeshEditor.Build.cs
--- a/Engine/Plugins/Editor/MeshEditor/Source/MeshEditor/MeshEditor.Build.cs
+++ b/Engine/Plugins/Editor/MeshEditor/Source/MeshEditor/MeshEditor.Build.cs
@@ -17,7 +17,7 @@ namespace UnrealBuildTool.Rules
                     "EditorStyle",
 				    "UnrealEd",
 					"RenderCore",
-					"MeshEditingRuntime",
+					"EditableMesh",
 					"ViewportInteraction",
                     "VREditor",
 					"Projects"
diff --git a/Engine/Plugins/Editor/MeshEditor/Source/PolygonModeling/PolygonModeling.Build.cs b/Engine/Plugins/Editor/MeshEditor/Source/PolygonModeling/PolygonModeling.Build.cs
--- a/Engine/Plugins/Editor/MeshEditor/Source/PolygonModeling/PolygonModeling.Build.cs
+++ b/Engine/Plugins/Editor/MeshEditor/Source/PolygonModeling/PolygonModeling.Build.cs
@@ -14,7 +14,7 @@ namespace UnrealBuildTool.Rules
 					"Engine",
 					"InputCore",
 					"UnrealEd",
-					"MeshEditingRuntime",
+					"EditableMesh",
 					"MeshEditor",
 					"SlateCore",
 					"ViewportInteraction"
diff --git a/Engine/Plugins/Runtime/EditableMesh/Source/EditableMesh/EditableMesh.Build.cs b/Engine/Plugins/Runtime/EditableMesh/Source/EditableMesh/EditableMesh.Build.cs
new file mode 100644
--- /dev/null
+++ b/Engine/Plugins/Runtime/EditableMesh/Source/EditableMesh/EditableMesh.Build.cs
@@ -0,0 +1,21 @@
+// Copyright 1998-2015 Epic Games, Inc. All Rights Reserved.
+
+namespace UnrealBuildTool.Rules
+{
+	public class EditableMesh : ModuleRules
+	{
+        public EditableMesh(ReadOnlyTargetRules Target) : base(Target)
+		{
+			PrivateDependencyModuleNames.AddRange(
+                new string[] {
+				    "Core",
+				    "CoreUObject",
+				    "Engine",
+					"RenderCore",	// @todo mesheditor: For FlushRenderingCommands()
+				}
+			);
+
+			AddEngineThirdPartyPrivateStaticDependencies(Target, new string[] { "MikkTSpace", "OpenSubdiv" });
+		}
+	}
+}
\ No newline at end of file
