Mesh Editor: New "EditableMesh" plugin (replaces "MeshEditingRuntime" module)

diff --git a/Engine/Plugins/Runtime/EditableMesh/Source/EditableMesh/EditableMesh.cpp b/Engine/Plugins/Runtime/EditableMesh/Source/EditableMesh/EditableMesh.cpp
new file mode 100644
--- /dev/null
+++ b/Engine/Plugins/Runtime/EditableMesh/Source/EditableMesh/EditableMesh.cpp
@@ -0,0 +1,9081 @@
+// Copyright 1998-2015 Epic Games, Inc. All Rights Reserved.
+
+#include "EditableMesh.h"
+#include "EditableMeshChanges.h"
+#include "EditableMeshCustomVersion.h"
+#include "EditableMeshAdapter.h"
+#include "HAL/IConsoleManager.h"
+#include "ProfilingDebugging/ScopedTimers.h"	// For FAutoScopedDurationTimer
+#include "GeomTools.h"
+#include "mikktspace.h"	// For tangent computations
+
+// =========================================================
+// OpenSubdiv support
+// =========================================================
+
+#if defined(_MSC_VER)
+	#pragma warning(push)
+	#pragma warning(disable:4191)		// Disable warning C4191: 'type cast' : unsafe conversion
+#endif
+
+#define M_PI PI		// OpenSubdiv is expecting M_PI to be defined already
+
+#include "far/topologyRefinerFactory.h"
+#include "far/topologyDescriptor.h"
+#include "far/topologyRefiner.h"
+#include "far/primvarRefiner.h"
+	
+#undef M_PI
+
+#if defined(_MSC_VER)
+	#pragma warning(pop)
+#endif
+
+DEFINE_LOG_CATEGORY( LogEditableMesh );
+
+// =========================================================
+
+
+namespace EditableMesh
+{
+	static FAutoConsoleVariable InterpolatePositionsToLimit( TEXT( "EditableMesh.InterpolatePositionsToLimit" ), 1, TEXT( "Whether to interpolate vertex positions for subdivision meshes all the way to their limit surface position.  Otherwise, we stop at the most refined mesh position." ) );
+	static FAutoConsoleVariable InterpolateFVarsToLimit( TEXT( "EditableMesh.InterpolateFVarsToLimit" ), 1, TEXT( "Whether to interpolate face-varying vertex data for subdivision meshes all the way to their limit surface position.  Otherwise, we stop at the most refined mesh." ) );
+}
+
+
+//
+// =========================================================
+//
+
+const FElementID FElementID::Invalid( TNumericLimits<uint32>::Max() );
+const FVertexID FVertexID::Invalid( TNumericLimits<uint32>::Max() );
+const FVertexInstanceID FVertexInstanceID::Invalid( TNumericLimits<uint32>::Max() );
+const FEdgeID FEdgeID::Invalid( TNumericLimits<uint32>::Max() );
+const FPolygonGroupID FPolygonGroupID::Invalid( TNumericLimits<uint32>::Max() );
+const FPolygonID FPolygonID::Invalid( TNumericLimits<uint32>::Max() );
+
+const FName UEditableMeshAttribute::VertexPositionName( "VertexPosition" );
+const FName UEditableMeshAttribute::VertexCornerSharpnessName( "VertexCornerSharpness" );
+const FName UEditableMeshAttribute::VertexNormalName( "VertexNormal" );
+const FName UEditableMeshAttribute::VertexTangentName( "VertexTangent" );
+const FName UEditableMeshAttribute::VertexBinormalSignName( "VertexBinormalSign" );
+const FName UEditableMeshAttribute::VertexTextureCoordinateName( "VertexTextureCoordinate" );
+const FName UEditableMeshAttribute::VertexColorName( "VertexColor" );
+const FName UEditableMeshAttribute::EdgeIsHardName( "EdgeIsHard" );
+const FName UEditableMeshAttribute::EdgeCreaseSharpnessName( "EdgeCreaseSharpness" );
+const FName UEditableMeshAttribute::PolygonNormalName( "PolygonNormal" );
+const FName UEditableMeshAttribute::PolygonCenterName( "PolygonCenter" );
+
+
+UEditableMesh::UEditableMesh()
+	: PendingCompactCounter( 0 )
+{
+}
+
+
+void UEditableMesh::Serialize( FArchive& Ar )
+{
+	Super::Serialize( Ar );
+	Ar.UsingCustomVersion( FEditableMeshCustomVersion::GUID );
+
+	SerializeSparseArray( Ar, Vertices );
+	SerializeSparseArray( Ar, VertexInstances );
+	SerializeSparseArray( Ar, Edges );
+	SerializeSparseArray( Ar, Polygons );
+	SerializeSparseArray( Ar, PolygonGroups );
+}
+
+
+void UEditableMesh::PostLoad()
+{
+	Super::PostLoad();
+
+	if( IsPreviewingSubdivisions() )
+	{
+		RefreshOpenSubdiv();
+	}
+
+	RebuildRenderMesh();
+}
+
+
+template <typename ElementIDType>
+static void InvertRemapTable( TSparseArray<ElementIDType>& InvertedRemapTable, const TSparseArray<ElementIDType>& RemapTable )
+{
+	static_assert( TIsDerivedFrom<ElementIDType, FElementID>::IsDerived, "Remap array type must be derived from FElementID" );
+
+	InvertedRemapTable.Empty( RemapTable.Num() );
+
+	for( TSparseArray<ElementIDType>::TConstIterator It( RemapTable ); It; ++It )
+	{
+		InvertedRemapTable.Insert( It->GetValue(), ElementIDType( It.GetIndex() ) );
+	}
+}
+
+
+class FCompactChange : public FChange
+{
+public:
+
+	/** Constructor */
+	FCompactChange()
+	{
+	}
+
+	// Parent class overrides
+	virtual TUniquePtr<FChange> Execute( UObject* Object ) override
+	{
+		UEditableMesh* EditableMesh = CastChecked<UEditableMesh>( Object );
+		verify( !EditableMesh->AnyChangesToUndo() );
+		EditableMesh->Compact();
+		return EditableMesh->MakeUndo();
+	}
+
+	virtual FString ToString() const override
+	{
+		return FString( TEXT( "Compact" ) );
+	}
+};
+
+
+struct FUncompactChangeInput
+{
+	/** A set of remap tables, specifying how the elements should have their indices remapped */
+	FElementIDRemappings ElementIDRemappings;
+};
+
+
+class FUncompactChange : public FChange
+{
+public:
+
+	/** Constructor */
+	FUncompactChange( const FUncompactChangeInput& InitInput )
+		: Input( InitInput )
+	{
+	}
+
+	FUncompactChange( FUncompactChangeInput&& InitInput )
+		: Input( MoveTemp( InitInput ) )
+	{
+	}
+
+	// Parent class overrides
+	virtual TUniquePtr<FChange> Execute( UObject* Object ) override
+	{
+		UEditableMesh* EditableMesh = CastChecked<UEditableMesh>( Object );
+		verify( !EditableMesh->AnyChangesToUndo() );
+		EditableMesh->Uncompact( Input.ElementIDRemappings );
+		return EditableMesh->MakeUndo();
+	}
+
+	virtual FString ToString() const override
+	{
+		return FString( TEXT( "Uncompact" ) );
+	}
+
+private:
+
+	/** The data we need to make this change */
+	FUncompactChangeInput Input;
+};
+
+
+void UEditableMesh::FixUpElementIDs( const FElementIDRemappings& Remappings )
+{
+	for( FMeshVertex& Vertex : Vertices )
+	{
+		// Fix up vertex instance index references in vertices array
+		for( FVertexInstanceID& VertexInstanceID : Vertex.VertexInstanceIDs )
+		{
+			VertexInstanceID = Remappings.GetRemappedVertexInstanceID( VertexInstanceID );
+		}
+
+		// Fix up edge index references in the vertex array
+		for( FEdgeID& EdgeID : Vertex.ConnectedEdgeIDs )
+		{
+			EdgeID = Remappings.GetRemappedEdgeID( EdgeID );
+		}
+	}
+
+	// Fix up vertex index references in vertex instance array
+	for( FMeshVertexInstance& VertexInstance : VertexInstances )
+	{
+		VertexInstance.VertexID = Remappings.GetRemappedVertexID( VertexInstance.VertexID );
+
+		for( FPolygonID& PolygonID : VertexInstance.ConnectedPolygons )
+		{
+			PolygonID = Remappings.GetRemappedPolygonID( PolygonID );
+		}
+	}
+
+	for( FMeshEdge& Edge : Edges )
+	{
+		// Fix up vertex index references in Edges array
+		for( int32 Index = 0; Index < 2; Index++ )
+		{
+			Edge.VertexIDs[ Index ] = Remappings.GetRemappedVertexID( Edge.VertexIDs[ Index ] );
+		}
+
+		// Fix up references to section indices
+		for( FPolygonID& ConnectedPolygon : Edge.ConnectedPolygons )
+		{
+			ConnectedPolygon = Remappings.GetRemappedPolygonID( ConnectedPolygon );
+		}
+	}
+
+	for( FMeshPolygon& Polygon : Polygons )
+	{
+		// Fix up references to vertex indices in section polygons' contours
+		for( FVertexInstanceID& VertexInstanceID : Polygon.PerimeterContour.VertexInstanceIDs )
+		{
+			VertexInstanceID = Remappings.GetRemappedVertexInstanceID( VertexInstanceID );
+		}
+
+		for( FMeshPolygonContour& HoleContour : Polygon.HoleContours )
+		{
+			for( FVertexInstanceID& VertexInstanceID : HoleContour.VertexInstanceIDs )
+			{
+				VertexInstanceID = Remappings.GetRemappedVertexInstanceID( VertexInstanceID );
+			}
+		}
+
+		for( FMeshTriangle& Triangle : Polygon.Triangles )
+		{
+			for( int32 TriangleVertexNumber = 0; TriangleVertexNumber < 3; ++TriangleVertexNumber )
+			{
+				const FVertexInstanceID OriginalVertexInstanceID = Triangle.GetVertexInstanceID( TriangleVertexNumber );
+				const FVertexInstanceID NewVertexInstanceID = Remappings.GetRemappedVertexInstanceID( OriginalVertexInstanceID );
+				Triangle.SetVertexInstanceID( TriangleVertexNumber, NewVertexInstanceID );
+			}
+		}
+
+		Polygon.PolygonGroupID = Remappings.GetRemappedPolygonGroupID( Polygon.PolygonGroupID );
+	}
+
+	for( FMeshPolygonGroup& PolygonGroup : PolygonGroups )
+	{
+		for( FPolygonID& Polygon : PolygonGroup.Polygons )
+		{
+			Polygon = Remappings.GetRemappedPolygonID( Polygon );
+		}
+	}
+}
+
+
+void UEditableMesh::Compact()
+{
+	static FElementIDRemappings Remappings;
+
+	CompactSparseArrayElements( Vertices, Remappings.NewVertexIndexLookup );
+	CompactSparseArrayElements( VertexInstances, Remappings.NewVertexInstanceIndexLookup );
+	CompactSparseArrayElements( Edges, Remappings.NewEdgeIndexLookup );
+	CompactSparseArrayElements( Polygons, Remappings.NewPolygonIndexLookup );
+	CompactSparseArrayElements( PolygonGroups, Remappings.NewPolygonGroupIndexLookup );
+
+	FixUpElementIDs( Remappings );
+
+	for( UEditableMeshAdapter* Adapter : Adapters )
+	{
+		Adapter->OnReindexElements( this, Remappings );
+	}
+
+	// @todo mesheditor: broadcast event with remappings so that any cached element IDs can be fixed up.
+
+	RebuildRenderMesh();
+
+	// Prepare the inverse transaction to reverse the compaction
+	FUncompactChangeInput UncompactChangeInput;
+	InvertRemapTable( UncompactChangeInput.ElementIDRemappings.NewVertexIndexLookup, Remappings.NewVertexIndexLookup );
+	InvertRemapTable( UncompactChangeInput.ElementIDRemappings.NewVertexInstanceIndexLookup, Remappings.NewVertexInstanceIndexLookup );
+	InvertRemapTable( UncompactChangeInput.ElementIDRemappings.NewEdgeIndexLookup, Remappings.NewEdgeIndexLookup );
+	InvertRemapTable( UncompactChangeInput.ElementIDRemappings.NewPolygonIndexLookup, Remappings.NewPolygonIndexLookup );
+	InvertRemapTable( UncompactChangeInput.ElementIDRemappings.NewPolygonGroupIndexLookup, Remappings.NewPolygonGroupIndexLookup );
+
+	AddUndo( MakeUnique<FUncompactChange>( MoveTemp( UncompactChangeInput ) ) );
+}
+
+
+void UEditableMesh::Uncompact( const FElementIDRemappings& Remappings )
+{
+	RemapSparseArrayElements( Vertices, Remappings.NewVertexIndexLookup );
+	RemapSparseArrayElements( VertexInstances, Remappings.NewVertexInstanceIndexLookup );
+	RemapSparseArrayElements( Edges, Remappings.NewEdgeIndexLookup );
+	RemapSparseArrayElements( Polygons, Remappings.NewPolygonIndexLookup );
+	RemapSparseArrayElements( PolygonGroups, Remappings.NewPolygonGroupIndexLookup );
+
+	FixUpElementIDs( Remappings );
+
+	for( UEditableMeshAdapter* Adapter : Adapters )
+	{
+		Adapter->OnReindexElements( this, Remappings );
+	}
+
+	// @todo mesheditor: broadcast event with remappings so that any cached element IDs can be fixed up.
+
+	RebuildRenderMesh();
+
+	AddUndo( MakeUnique<FCompactChange>() );
+}
+
+
+void UEditableMesh::RebuildRenderMesh()
+{
+	if( !IsBeingModified() )
+	{
+		const bool bRefreshBounds = true;
+		const bool bInvalidateLighting = true;
+		for( UEditableMeshAdapter* Adapter : Adapters )
+		{
+			Adapter->OnRebuildRenderMeshStart( this, bRefreshBounds, bInvalidateLighting );
+		}
+	}
+
+	for( UEditableMeshAdapter* Adapter : Adapters )
+	{
+		Adapter->OnRebuildRenderMesh( this );
+	}
+
+	if( !IsBeingModified() )
+	{
+		const bool bUpdateCollision = true;
+		for( UEditableMeshAdapter* Adapter : Adapters )
+		{
+			Adapter->OnRebuildRenderMeshFinish( this, bUpdateCollision );
+		}
+	}
+}
+
+
+void UEditableMesh::StartModification( const EMeshModificationType MeshModificationType, const EMeshTopologyChange MeshTopologyChange )
+{
+	if( ensure( !IsBeingModified() ) )
+	{
+		bIsBeingModified = true;
+
+		// Should be nothing in the undo stack if we're just starting to modify the mesh now
+		ensure( !this->AnyChangesToUndo() );
+
+		FStartOrEndModificationChangeInput RevertInput;
+		RevertInput.bStartModification = false;
+		RevertInput.MeshModificationType = MeshModificationType;
+		RevertInput.MeshTopologyChange = MeshTopologyChange;
+		AddUndo( MakeUnique<FStartOrEndModificationChange>( MoveTemp( RevertInput ) ) );
+
+		this->CurrentModificationType = MeshModificationType;
+		this->CurrentToplogyChange = MeshTopologyChange;
+
+		// @todo mesheditor debug: Disable noisy mesh editor spew by default (here and elsewhere)
+		// UE_LOG( LogEditableMesh, Log, TEXT( "UEditableStaticMesh::StartModification START: %s" ), *SubMeshAddress.ToString() );
+		FAutoScopedDurationTimer FunctionTimer;
+
+		const bool bRefreshBounds = MeshModificationType == EMeshModificationType::Final;	 // @todo mesheditor perf: Only do this if we may have changed the bounds
+		const bool bInvalidateLighting = ( MeshModificationType == EMeshModificationType::FirstInterim || MeshModificationType == EMeshModificationType::Final );	// @todo mesheditor perf: We can avoid invalidating lighting on 'Final' if we know that a 'FirstInterim' happened since the last 'Final'
+		for( UEditableMeshAdapter* Adapter : Adapters )
+		{
+			Adapter->OnRebuildRenderMeshStart( this, bRefreshBounds, bInvalidateLighting );
+		}
+
+		for( UEditableMeshAdapter* Adapter : Adapters )
+		{
+			Adapter->OnStartModification( this, MeshModificationType, MeshTopologyChange );
+		}
+
+		PolygonsPendingNewTangentBasis.Reset();
+		PolygonsPendingTriangulation.Reset();
+
+		// @todo mesheditor debug
+		// UE_LOG( LogEditableMesh, Log, TEXT( "UEditableStaticMesh::StartModification COMPLETE in %0.4fs" ), FunctionTimer.GetTime() );
+	}
+}
+
+
+void UEditableMesh::EndModification( const bool bFromUndo )
+{
+	if( ensure( IsBeingModified() ) )
+	{
+		// @todo mesheditor debug
+		// UE_LOG( LogEditableMesh, Log, TEXT( "UEditableStaticMesh::EndModification START (ModType=%i): %s" ), (int32)MeshModificationType, *SubMeshAddress.ToString() );
+		// FAutoScopedDurationTimer FunctionTimer;
+
+		// Retriangulate first, as the triangulation of n-gons determines how the tangent basis is calculated
+		if( PolygonsPendingTriangulation.Num() > 0 )
+		{
+			RetriangulatePolygons();
+		}
+
+		if( PolygonsPendingNewTangentBasis.Num() > 0 )
+		{
+			GenerateTangentsAndNormals();
+		}
+
+		if( CurrentModificationType == EMeshModificationType::Final || !bFromUndo )
+		{
+			// Update subdivision limit surface
+			if( CurrentToplogyChange == EMeshTopologyChange::TopologyChange )
+			{
+				// Mesh topology (or subdivision level or smoothing) may have changed, so go ahead and refresh our OpenSubdiv representation entirely
+				RefreshOpenSubdiv();
+			}
+			else
+			{
+				// No topology change, so we can ask OpenSubdiv to quickly generate new limit surface geometry
+				GenerateOpenSubdivLimitSurfaceData();
+			}
+		}
+
+		// Every so often, compact the data.
+		// Note we only want to do this when actions are performed, not when they are being undone/redone
+		bool bDidCompact = false;
+
+		if( CurrentModificationType == EMeshModificationType::Final &&
+			CurrentToplogyChange == EMeshTopologyChange::TopologyChange &&
+			!bFromUndo )
+		{
+			if( ++PendingCompactCounter == CompactFrequency )
+			{
+				PendingCompactCounter = 0;
+				Compact();
+				bDidCompact = true;
+			}
+		}
+
+		// If subdivision preview mode is active, we'll need to refresh the entire static mesh with data from OpenSubdiv
+		// @todo mesheditor subdiv perf: Ideally we can avoid refreshing the entire thing if only positions have changed, as per above
+		if( IsPreviewingSubdivisions() && ( CurrentModificationType == EMeshModificationType::Final || !bFromUndo ) )
+		{
+			if( !bDidCompact )	// If we did a Compact() in this function, the mesh will have already been rebuilt
+			{
+				for( UEditableMeshAdapter* Adapter : Adapters )
+				{
+					Adapter->OnRebuildRenderMesh( this );
+				}
+			}
+		}
+
+		for( UEditableMeshAdapter* Adapter : Adapters )
+		{
+			Adapter->OnRebuildRenderMeshFinish( this, CurrentModificationType == EMeshModificationType::Final );
+		}
+
+		for( UEditableMeshAdapter* Adapter : Adapters )
+		{
+			Adapter->OnEndModification( this );
+		}
+
+		// @todo mesheditor: Not currently sure if we need to do this or not.  Also, we are trying to support runtime editing (no PostEditChange stuff)
+		// 		if( GIsEditor && MeshModificationType == EMeshModificationType::Final )
+		// 		{
+		// 			StaticMesh->PostEditChange();
+		// 		}
+
+		// @todo mesheditor debug
+		// UE_LOG( LogEditableMesh, Log, TEXT( "UEditableStaticMesh::EndModification COMPLETE in %0.4fs" ), FunctionTimer.GetTime() );	  // @todo mesheditor: Shows bogus time values
+
+		FStartOrEndModificationChangeInput RevertInput;
+		RevertInput.bStartModification = true;
+		RevertInput.MeshModificationType = CurrentModificationType;
+		RevertInput.MeshTopologyChange = CurrentToplogyChange;
+		AddUndo( MakeUnique<FStartOrEndModificationChange>( MoveTemp( RevertInput ) ) );
+
+		bIsBeingModified = false;
+	}
+}
+
+
+bool UEditableMesh::IsCommitted() const
+{
+	for( UEditableMeshAdapter* Adapter : Adapters )
+	{
+		if( !Adapter->IsCommitted( this ) )
+		{
+			return false;
+		}
+	}
+
+	return true;
+}
+
+
+bool UEditableMesh::IsCommittedAsInstance() const
+{
+	for( UEditableMeshAdapter* Adapter : Adapters )
+	{
+		if( !Adapter->IsCommittedAsInstance( this ) )
+		{
+			return false;
+		}
+	}
+
+	return true;
+}
+
+
+void UEditableMesh::Commit()
+{
+	for( UEditableMeshAdapter* Adapter : Adapters )
+	{
+		Adapter->OnCommit( this );
+	}
+}
+
+
+UEditableMesh* UEditableMesh::CommitInstance( UPrimitiveComponent* ComponentToInstanceTo )
+{
+	UEditableMesh* Result = nullptr;
+	for( UEditableMeshAdapter* Adapter : Adapters )
+	{
+		UEditableMesh* Instance = Adapter->OnCommitInstance( this, ComponentToInstanceTo );
+		if( !Result )
+		{
+			Result = Instance;
+		}
+	}
+
+	return Result;
+}
+
+
+// @todo mesheditor: implement these, or - better - reconsider the editing paradigm
+void UEditableMesh::Revert() {}
+UEditableMesh* UEditableMesh::RevertInstance() { return nullptr; }
+void UEditableMesh::PropagateInstanceChanges() {}
+
+
+const FEditableMeshSubMeshAddress& UEditableMesh::GetSubMeshAddress() const
+{
+	return SubMeshAddress;
+}
+
+
+void UEditableMesh::SetSubMeshAddress( const FEditableMeshSubMeshAddress& NewSubMeshAddress )
+{
+	SubMeshAddress = NewSubMeshAddress;
+}
+
+
+int32 UEditableMesh::GetVertexCount() const
+{
+	return Vertices.Num();
+}
+
+
+int32 UEditableMesh::GetVertexArraySize() const
+{
+	return Vertices.GetMaxIndex();
+}
+
+
+bool UEditableMesh::IsValidVertex( const FVertexID VertexID ) const
+{
+	return VertexID.GetValue() >= 0 && VertexID.GetValue() < Vertices.GetMaxIndex() && Vertices.IsAllocated( VertexID.GetValue() );
+}
+
+
+FVector4 UEditableMesh::GetVertexAttribute( const FVertexID VertexID, const FName AttributeName, const int32 AttributeIndex ) const
+{
+	const FMeshVertex& Vertex = Vertices[ VertexID.GetValue() ];
+
+	if( AttributeName == UEditableMeshAttribute::VertexPosition() )
+	{
+		check( AttributeIndex == 0 );	// Only one position is supported
+		return FVector4( Vertex.VertexPosition, 0.0f );
+	}
+	else if( AttributeName == UEditableMeshAttribute::VertexCornerSharpness() )
+	{
+		check( AttributeIndex == 0 );	// Only one softness value is supported
+		return FVector4( Vertex.CornerSharpness, 0.0f, 0.0f, 0.0f );
+	}
+
+	checkf( 0, TEXT( "UEditableMesh::GetVertexAttribute() called with unrecognized vertex attribute name: %s (index: %i)" ), *AttributeName.ToString(), AttributeIndex );
+	return FVector4( 0.0f );
+}
+
+
+void UEditableMesh::SetVertexAttribute( const FVertexID VertexID, const FName AttributeName, const int32 AttributeIndex, const FVector4 AttributeValue )
+{
+	FMeshVertex& Vertex = Vertices[ VertexID.GetValue() ];
+
+	if( AttributeName == UEditableMeshAttribute::VertexPosition() )
+	{
+		check( AttributeIndex == 0 );	// Only one position is supported
+		Vertex.VertexPosition = AttributeValue;
+	}
+	else if( AttributeName == UEditableMeshAttribute::VertexCornerSharpness() )
+	{
+		check( AttributeIndex == 0 );	// Only one softness value is supported
+		Vertex.CornerSharpness = AttributeValue.X;
+	}
+	else
+	{
+		checkf( 0, TEXT( "UEditableMesh::SetVertexAttribute() called with unrecognized vertex attribute name: %s (index: %i)" ), *AttributeName.ToString(), AttributeIndex );
+	}
+
+	for( UEditableMeshAdapter* Adapter : Adapters )
+	{
+		Adapter->OnSetVertexAttribute( this, VertexID, AttributeName, AttributeIndex, AttributeValue );
+	}
+}
+
+
+int32 UEditableMesh::GetVertexConnectedEdgeCount( const FVertexID VertexID ) const
+{
+	checkSlow( Vertices.IsAllocated( VertexID.GetValue() ) );
+	const FMeshVertex& Vertex = Vertices[ VertexID.GetValue() ];
+	return Vertex.ConnectedEdgeIDs.Num();
+}
+
+
+FEdgeID UEditableMesh::GetVertexConnectedEdge( const FVertexID VertexID, const int32 ConnectedEdgeNumber ) const
+{
+	checkSlow( Vertices.IsAllocated( VertexID.GetValue() ) );
+	const FMeshVertex& Vertex = Vertices[ VertexID.GetValue() ];
+	return Vertex.ConnectedEdgeIDs[ ConnectedEdgeNumber ];
+}
+
+
+int32 UEditableMesh::GetVertexInstanceCount() const
+{
+	return VertexInstances.Num();
+}
+
+
+int32 UEditableMesh::GetVertexInstanceArraySize() const
+{
+	return VertexInstances.GetMaxIndex();
+}
+
+
+FVector4 UEditableMesh::GetVertexInstanceAttribute( const FVertexInstanceID VertexInstanceID, const FName AttributeName, const int32 AttributeIndex ) const
+{
+	const FMeshVertexInstance& VertexInstance = VertexInstances[ VertexInstanceID.GetValue() ];
+
+	if( AttributeName == UEditableMeshAttribute::VertexPosition() )
+	{
+		check( AttributeIndex == 0 );	// Only one position is supported
+		return GetVertexAttribute( VertexInstance.VertexID, AttributeName, AttributeIndex );
+	}
+	else if( AttributeName == UEditableMeshAttribute::VertexNormal() )
+	{
+		check( AttributeIndex == 0 );	// Only one normal is supported
+		return FVector4( VertexInstance.Normal, 0.0f );
+	}
+	else if( AttributeName == UEditableMeshAttribute::VertexTangent() )
+	{
+		check( AttributeIndex == 0 );	// Only one tangent is supported
+		return FVector4( VertexInstance.Tangent, 0.0f );
+	}
+	else if( AttributeName == UEditableMeshAttribute::VertexBinormalSign() )
+	{
+		check( AttributeIndex == 0 );	// Only one basis determinant sign is supported
+		return FVector4( VertexInstance.BinormalSign );
+	}
+	else if( AttributeName == UEditableMeshAttribute::VertexTextureCoordinate() )
+	{
+		const int32 TextureCoordinateIndex = AttributeIndex;
+		if( TextureCoordinateIndex < VertexInstance.VertexUVs.Num() )
+		{
+			const FVector2D TextureCoordinate = VertexInstance.VertexUVs[ TextureCoordinateIndex ];
+			return FVector4( TextureCoordinate, FVector2D::ZeroVector );
+		}
+		else
+		{
+			return FVector4( 0.0f );
+		}
+	}
+	else if( AttributeName == UEditableMeshAttribute::VertexColor() )
+	{
+		check( AttributeIndex == 0 );
+		return FVector4( FLinearColor( VertexInstance.Color ) );
+	}
+
+	checkf( 0, TEXT( "UEditableMesh::GetVertexInstanceAttribute() called with unrecognized vertex instance attribute name: %s (index: %i)" ), *AttributeName.ToString(), AttributeIndex );
+	return FVector4( 0.0f );
+}
+
+
+void UEditableMesh::SetVertexInstanceAttribute( const FVertexInstanceID VertexInstanceID, const FName AttributeName, const int32 AttributeIndex, const FVector4 AttributeValue )
+{
+	FMeshVertexInstance& VertexInstance = VertexInstances[ VertexInstanceID.GetValue() ];
+
+	if( AttributeName == UEditableMeshAttribute::VertexPosition() )
+	{
+		check( AttributeIndex == 0 );	// Only one position is supported
+		SetVertexAttribute( VertexInstance.VertexID, AttributeName, AttributeIndex, AttributeValue );
+	}
+	else if( AttributeName == UEditableMeshAttribute::VertexNormal() )
+	{
+		check( AttributeIndex == 0 );	// Only one normal is supported
+		VertexInstance.Normal = AttributeValue;
+	}
+	else if( AttributeName == UEditableMeshAttribute::VertexTangent() )
+	{
+		check( AttributeIndex == 0 );	// Only one tangent is supported
+		VertexInstance.Tangent = AttributeValue;
+	}
+	else if( AttributeName == UEditableMeshAttribute::VertexBinormalSign() )
+	{
+		check( AttributeIndex == 0 );	// Only one basis determinant sign is supported
+		VertexInstance.BinormalSign = AttributeValue.X;
+	}
+	else if( AttributeName == UEditableMeshAttribute::VertexTextureCoordinate() )
+	{
+		if( AttributeIndex < VertexInstance.VertexUVs.Num() )
+		{
+			VertexInstance.VertexUVs[ AttributeIndex ] = FVector2D( AttributeValue.X, AttributeValue.Y );
+		}
+	}
+	else if( AttributeName == UEditableMeshAttribute::VertexColor() )
+	{
+		check( AttributeIndex == 0 );
+		VertexInstance.Color = FLinearColor( AttributeValue.X, AttributeValue.Y, AttributeValue.Z, AttributeValue.W );
+	}
+	else
+	{
+		checkf( 0, TEXT( "UEditableMesh::SetVertexInstanceAttribute() called with unrecognized vertex instance attribute name: %s (index: %i)" ), *AttributeName.ToString(), AttributeIndex );
+	}
+
+	for( UEditableMeshAdapter* Adapter : Adapters )
+	{
+		Adapter->OnSetVertexInstanceAttribute( this, VertexInstanceID, AttributeName, AttributeIndex, AttributeValue );
+	}
+}
+
+
+FVector4 UEditableMesh::GetEdgeAttribute( const FEdgeID EdgeID, const FName AttributeName, const int32 AttributeIndex ) const
+{
+	checkSlow( Edges.IsAllocated( EdgeID.GetValue() ) );
+	const FMeshEdge& Edge = Edges[ EdgeID.GetValue() ];
+
+	if( AttributeName == UEditableMeshAttribute::EdgeIsHard() )
+	{
+		check( AttributeIndex == 0 );	// Only one edge is hard flag is supported
+
+		return FVector4( Edge.bIsHardEdge ? 1.0f : 0.0f, 0.0f, 0.0f, 0.0f );
+	}
+	else if( AttributeName == UEditableMeshAttribute::EdgeCreaseSharpness() )
+	{
+		check( AttributeIndex == 0 );	// Only one edge crease sharpness is supported
+
+		return FVector4( Edge.CreaseSharpness, 0.0f, 0.0f, 0.0f );
+	}
+
+	checkf( 0, TEXT( "UEditableMesh::GetEdgeAttribute() called with unrecognized edge attribute name: %s (index: %i)" ), *AttributeName.ToString(), AttributeIndex );
+	return FVector4( 0.0f );
+}
+
+
+void UEditableMesh::SetEdgeAttribute( const FEdgeID EdgeID, const FName AttributeName, const int32 AttributeIndex, const FVector4 AttributeValue, bool bIsUndo )
+{
+	checkSlow( Edges.IsAllocated( EdgeID.GetValue() ) );
+	FMeshEdge& Edge = Edges[ EdgeID.GetValue() ];
+
+	if( AttributeName == UEditableMeshAttribute::EdgeIsHard() )
+	{
+		check( AttributeIndex == 0 );	// Only one edge is hard flag is supported
+		SetEdgeHardness( EdgeID, AttributeValue.X != 0.0f, bIsUndo );
+
+		// @todo mesheditor: this needs to potentially merge/split vertex instances and rebuild normals
+	}
+	else if( AttributeName == UEditableMeshAttribute::EdgeCreaseSharpness() )
+	{
+		check( AttributeIndex == 0 );	// Only one edge crease sharpness is supported
+		Edge.CreaseSharpness = AttributeValue.X;
+	}
+	else
+	{
+		checkf( 0, TEXT( "UEditableMesh::SetEdgeAttribute() called with unrecognized edge attribute name: %s (index: %i)" ), *AttributeName.ToString(), AttributeIndex );
+	}
+
+	for( UEditableMeshAdapter* Adapter : Adapters )
+	{
+		Adapter->OnSetEdgeAttribute( this, EdgeID, AttributeName, AttributeIndex, AttributeValue );
+	}
+}
+
+
+void UEditableMesh::GetPolygonsInSameSoftEdgedGroup( const FVertexID VertexID, const FPolygonID PolygonID, TArray<FPolygonID>& OutPolygonIDs ) const
+{
+	// The aim here is to determine which polygons form part of the same soft edged group as the polygons attached to this vertex instance.
+	// They should all contribute to the final vertex instance normal.
+
+	OutPolygonIDs.Reset();
+
+	// Get all polygons connected to this vertex.
+	static TArray<FPolygonID> ConnectedPolygons;
+	GetVertexConnectedPolygons( VertexID, ConnectedPolygons );
+
+	// Cache a list of all soft edges which share this vertex.
+	// We're only interested in finding adjacent polygons which are not the other side of a hard edge.
+	static TArray<FEdgeID> ConnectedSoftEdges;
+	GetConnectedSoftEdges( VertexID, ConnectedSoftEdges );
+
+	// Iterate through adjacent polygons which contain the given vertex, but don't cross a hard edge.
+	// Maintain a list of polygon IDs to be examined. Adjacents are added to the list if suitable.
+	// Add the start poly here.
+	static TArray<FPolygonID> PolygonsToCheck;
+	PolygonsToCheck.Reset();
+	PolygonsToCheck.Add( PolygonID );
+
+	int32 Index = 0;
+	while( Index < PolygonsToCheck.Num() )
+	{
+		const FPolygonID PolygonToCheck = PolygonsToCheck[ Index ];
+		Index++;
+
+		if( ConnectedPolygons.Contains( PolygonToCheck ) )
+		{
+			OutPolygonIDs.Add( PolygonToCheck );
+
+			// Now look at its adjacent polygons. If they are joined by a soft edge which includes the vertex we're interested in, we want to consider them.
+			// We take a shortcut by doing this process in reverse: we already know all the soft edges we are interested in, so check if any of them
+			// have the current polygon as an adjacent.
+			for( const FEdgeID ConnectedSoftEdge : ConnectedSoftEdges )
+			{
+				const FMeshEdge& Edge = Edges[ ConnectedSoftEdge.GetValue() ];
+				if( Edge.ConnectedPolygons.Contains( PolygonToCheck ) )
+				{
+					for( const FPolygonID AdjacentPolygon : Edge.ConnectedPolygons )
+					{
+						// Only add new polygons which haven't yet been added to the list. This prevents circular runs of polygons triggering infinite loops.
+						PolygonsToCheck.AddUnique( AdjacentPolygon );
+					}
+				}
+			}
+		}
+	}
+}
+
+
+void UEditableMesh::SplitVertexInstanceInPolygons( const FVertexInstanceID VertexInstanceID, const TArray<FPolygonID>& PolygonIDs )
+{
+	UE_LOG( LogEditableMesh, Verbose, TEXT( "SplitVertexInstanceInPolygons: %s %s" ), *VertexInstanceID.ToString(), *LogHelpers::ArrayToString( PolygonIDs ) );
+
+	const FMeshVertexInstance& VertexInstance = VertexInstances[ VertexInstanceID.GetValue() ];
+
+	// Create a new vertex instance copied from the one passed in
+
+	static TArray<FVertexInstanceToCreate> VertexInstancesToCreate;
+	VertexInstancesToCreate.Reset();
+	VertexInstancesToCreate.SetNum( 1 );
+	VertexInstancesToCreate[ 0 ].VertexID = VertexInstance.VertexID;
+
+	for( const FName AttributeName : UEditableMesh::GetValidVertexInstanceAttributes() )
+	{
+		const int32 MaxAttributeIndex = GetMaxAttributeIndex( AttributeName );
+		for( int32 AttributeIndex = 0; AttributeIndex < MaxAttributeIndex; ++AttributeIndex )
+		{
+			VertexInstancesToCreate[ 0 ].VertexInstanceAttributes.Attributes.Emplace(
+				AttributeName,
+				AttributeIndex,
+				GetVertexInstanceAttribute( VertexInstanceID, AttributeName, AttributeIndex )
+			);
+		}
+	}
+
+	static TArray<FVertexInstanceID> NewVertexInstanceIDs;
+	CreateVertexInstances( VertexInstancesToCreate, NewVertexInstanceIDs );
+
+	ReplaceVertexInstanceInPolygons( VertexInstanceID, NewVertexInstanceIDs[ 0 ], PolygonIDs );
+
+	UE_LOG( LogEditableMesh, Verbose, TEXT( "* SplitVertexInstanceInPolygons returned" ) );
+}
+
+
+void UEditableMesh::ReplaceVertexInstanceInPolygons( const FVertexInstanceID OldVertexInstanceID, const FVertexInstanceID NewVertexInstanceID, const TArray<FPolygonID>& PolygonIDs )
+{
+	UE_LOG( LogEditableMesh, Verbose, TEXT( "ReplaceVertexInstanceInPolygons: %s %s %s" ), *OldVertexInstanceID.ToString(), *NewVertexInstanceID.ToString(), *LogHelpers::ArrayToString( PolygonIDs ) );
+
+	// Substitute the new vertex instance in the passed in polygons
+
+	static TArray<FChangeVertexInstancesForPolygon>  VertexInstancesToChange;
+	VertexInstancesToChange.Reset( PolygonIDs.Num() );
+
+	for( const FPolygonID PolygonID : PolygonIDs )
+	{
+		VertexInstancesToChange.Emplace();
+		FChangeVertexInstancesForPolygon& ChangeVertexInstances = VertexInstancesToChange.Last();
+
+		ChangeVertexInstances.PolygonID = PolygonID;
+
+		const FMeshPolygon& Polygon = Polygons[ PolygonID.GetValue() ];
+		int32 VertexInstanceIndex = Polygon.PerimeterContour.VertexInstanceIDs.Find( OldVertexInstanceID );
+
+		if( VertexInstanceIndex != INDEX_NONE )
+		{
+			// Found the vertex instance in the perimeter
+			ChangeVertexInstances.PerimeterVertexIndicesAndInstanceIDs.Emplace();
+			ChangeVertexInstances.PerimeterVertexIndicesAndInstanceIDs[ 0 ].ContourIndex = VertexInstanceIndex;
+			ChangeVertexInstances.PerimeterVertexIndicesAndInstanceIDs[ 0 ].VertexInstanceID = NewVertexInstanceID;
+		}
+		else
+		{
+			ChangeVertexInstances.VertexIndicesAndInstanceIDsForEachHole.SetNum( Polygon.HoleContours.Num() );
+
+			int32 HoleIndex = 0;
+			for( const FMeshPolygonContour& HoleContour : Polygon.HoleContours )
+			{
+				VertexInstanceIndex = HoleContour.VertexInstanceIDs.Find( OldVertexInstanceID );
+				if( VertexInstanceIndex != INDEX_NONE )
+				{
+					// Found the vertex instance in a hole
+					ChangeVertexInstances.VertexIndicesAndInstanceIDsForEachHole[ HoleIndex ].VertexIndicesAndInstanceIDs.Emplace();
+					ChangeVertexInstances.VertexIndicesAndInstanceIDsForEachHole[ HoleIndex ].VertexIndicesAndInstanceIDs[ 0 ].ContourIndex = VertexInstanceIndex;
+					ChangeVertexInstances.VertexIndicesAndInstanceIDsForEachHole[ HoleIndex ].VertexIndicesAndInstanceIDs[ 0 ].VertexInstanceID = NewVertexInstanceID;
+					break;
+				}
+
+				HoleIndex++;
+			}
+		}
+
+		// We expect to have found the vertex instance somewhere in one of the polygon contours
+		check( VertexInstanceIndex != INDEX_NONE );
+	}
+
+	ChangePolygonsVertexInstances( VertexInstancesToChange );
+
+	UE_LOG( LogEditableMesh, Verbose, TEXT( "* ReplaceVertexInstanceInPolygons returned" ) );
+}
+
+
+void UEditableMesh::GetVertexInstanceConnectedPolygonsInSameGroup( const FVertexInstanceID VertexInstanceID, const FPolygonID PolygonID, TArray<FPolygonID>& OutPolygonIDs ) const
+{
+	// The aim here is to determine which polygons connected to this vertex instance form part of the same soft edged group as the specified polygon.
+
+	OutPolygonIDs.Reset();
+
+	const FMeshVertexInstance& VertexInstance = VertexInstances[ VertexInstanceID.GetValue() ];
+	const FVertexID VertexID = VertexInstance.VertexID;
+
+	// Cache a list of all soft edges which share this vertex.
+	// We're only interested in finding adjacent polygons which are not the other side of a hard edge.
+	static TArray<FEdgeID> ConnectedSoftEdges;
+	GetConnectedSoftEdges( VertexID, ConnectedSoftEdges );
+
+	// Iterate through adjacent polygons which contain the given vertex, but don't cross a hard edge.
+	// Maintain a list of polygon IDs to be examined. Adjacents are added to the list if suitable.
+	// Add the start poly here.
+	static TArray<FPolygonID> PolygonsToCheck;
+	PolygonsToCheck.Reset();
+	PolygonsToCheck.Add( PolygonID );
+
+	int32 Index = 0;
+	while( Index < PolygonsToCheck.Num() )
+	{
+		const FPolygonID PolygonToCheck = PolygonsToCheck[ Index ];
+		Index++;
+
+		if( VertexInstance.ConnectedPolygons.Contains( PolygonToCheck ) )
+		{
+			OutPolygonIDs.Add( PolygonToCheck );
+
+			// Now look at its adjacent polygons. If they are joined by a soft edge which includes the vertex we're interested in, we want to consider them.
+			// We take a shortcut by doing this process in reverse: we already know all the soft edges we are interested in, so check if any of them
+			// have the current polygon as an adjacent.
+			for( const FEdgeID ConnectedSoftEdge : ConnectedSoftEdges )
+			{
+				const FMeshEdge& Edge = Edges[ ConnectedSoftEdge.GetValue() ];
+				if( Edge.ConnectedPolygons.Contains( PolygonToCheck ) )
+				{
+					for( const FPolygonID AdjacentPolygon : Edge.ConnectedPolygons )
+					{
+						// Only add new polygons which haven't yet been added to the list. This prevents circular runs of polygons triggering infinite loops.
+						PolygonsToCheck.AddUnique( AdjacentPolygon );
+					}
+				}
+			}
+		}
+	}
+}
+
+
+void UEditableMesh::SetEdgeHardness( const FEdgeID EdgeID, const bool bIsHard, const bool bIsUndo )
+{
+	checkSlow( Edges.IsAllocated( EdgeID.GetValue() ) );
+	FMeshEdge& Edge = Edges[ EdgeID.GetValue() ];
+
+	if( Edge.bIsHardEdge != bIsHard )
+	{
+		Edge.bIsHardEdge = bIsHard;
+
+		if( !bIsUndo )
+		{
+			UE_LOG( LogEditableMesh, Verbose, TEXT( "SetEdgeHardness (changed): %s %s" ), *EdgeID.ToString(), *LogHelpers::BoolToString( bIsHard ) );
+
+			if( bIsHard )
+			{
+				// Setting a hard edge potentially requires vertex instances to be split
+
+				for( const FVertexID VertexID : Edge.VertexIDs )
+				{
+					const FMeshVertex& Vertex = Vertices[ VertexID.GetValue() ];
+
+					// For each end of the edge, go through all vertex instances, determining if any need to be split.
+					// Take a copy of the vertex instances because splitting them will mutate the list we are iterating.
+					static TArray<FVertexInstanceID> VertexInstanceIDsCopy;
+					VertexInstanceIDsCopy = Vertex.VertexInstanceIDs;
+
+					// Iterate all the vertex instances for each endpoint.
+					// We are looking for vertex instances which need splitting into two different instances
+					// (because the hard edge will require separate normals for each of its polygon smoothing groups).
+					for( const FVertexInstanceID VertexInstanceID : VertexInstanceIDsCopy )
+					{
+						const FMeshVertexInstance& VertexInstance = VertexInstances[ VertexInstanceID.GetValue() ];
+
+						// Take a copy of the array of polygons connected to this instance.
+						static TArray<FPolygonID> PolygonsSharingThisVertex;
+						PolygonsSharingThisVertex = VertexInstance.ConnectedPolygons;
+						bool bFirstTime = true;
+						while( PolygonsSharingThisVertex.Num() > 0 )
+						{
+							// For the next polygon in the array, determine all other polygons in the same smoothing group which share this vertex instance.
+							static TArray<FPolygonID> PolygonsInSameSmoothingGroup;
+							GetVertexInstanceConnectedPolygonsInSameGroup( VertexInstanceID, PolygonsSharingThisVertex[ 0 ], PolygonsInSameSmoothingGroup );
+
+							// Check that all polygons in the smoothing group are attached to this vertex instance, and remove them from the master list of polygons
+							// connected to this instance. If a polygon in the smoothing group is not attached to this vertex instance, it's because it's the other
+							// side of a UV seam and hence has a distinct vertex instance.
+							for( const FPolygonID PolygonInSmoothingGroup : PolygonsInSameSmoothingGroup )
+							{
+								verify( PolygonsSharingThisVertex.Remove( PolygonInSmoothingGroup ) == 1 );
+							}
+
+							// First group which we extract: do nothing - they can keep their existing instance ID.
+							// Subsequent times round the loop, we create a new vertex instance copied from the original one, and replace connected polygon vertices with it.
+							if( !bFirstTime )
+							{
+								SplitVertexInstanceInPolygons( VertexInstanceID, PolygonsInSameSmoothingGroup );
+							}
+
+							bFirstTime = false;
+						}
+					}
+				}
+			}
+			else
+			{
+				// Setting a soft edge potentially requires vertex instances to be merged
+
+				static TArray<FVertexInstanceID> VertexInstancesToDelete;
+				VertexInstancesToDelete.Reset();
+
+				for( const FVertexID VertexID : Edge.VertexIDs )
+				{
+					// For each end of the edge which has been softened...
+					const FMeshVertex& Vertex = Vertices[ VertexID.GetValue() ];
+
+					// ...iterate through pairs of vertex instances, looking for potential to merge them
+					for( int32 IndexA = 0; IndexA < Vertex.VertexInstanceIDs.Num() - 1; ++IndexA )
+					{
+						const FVertexInstanceID VertexInstanceIDA = Vertex.VertexInstanceIDs[ IndexA ];
+						const FMeshVertexInstance& VertexInstanceA = VertexInstances[ VertexInstanceIDA.GetValue() ];
+
+						// If vertex instance isn't connected to any polygon, we can't deduce anything about its smoothing group, so skip to the next one
+						if( VertexInstanceA.ConnectedPolygons.Num() == 0 )
+						{
+							continue;
+						}
+
+						// Get mergeable set of attributes for first vertex instance
+						static TArray<FMeshElementAttributeData> AttributesA;
+						GetMergedNamedVertexInstanceAttributeData( GetMergeableVertexInstanceAttributes(), Vertex.VertexInstanceIDs[ IndexA ], TArray<FMeshElementAttributeData>(), AttributesA );
+
+						for( int32 IndexB = IndexA + 1; IndexB < Vertex.VertexInstanceIDs.Num(); ++IndexB )
+						{
+							const FVertexInstanceID VertexInstanceIDB = Vertex.VertexInstanceIDs[ IndexB ];
+
+							// If this vertex instance has been marked for deletion in a previous iteration, skip it
+							if( VertexInstancesToDelete.Contains( VertexInstanceIDB ) )
+							{
+								continue;
+							}
+
+							// Get mergeable set of attributes for second vertex instance
+							static TArray<FMeshElementAttributeData> AttributesB;
+							GetMergedNamedVertexInstanceAttributeData( GetMergeableVertexInstanceAttributes(), Vertex.VertexInstanceIDs[ IndexB ], TArray<FMeshElementAttributeData>(), AttributesB );
+
+							// If both vertex instances have equal mergeable attributes, they can potentially be merged
+							if( AreAttributeListsEqual( AttributesA, AttributesB ) )
+							{
+								const FMeshVertexInstance& VertexInstanceB = VertexInstances[ VertexInstanceIDB.GetValue() ];
+
+								// Check to see if they are in the same smoothing group, using an arbitary connected polygon in the first vertex instance as the pivot point.
+								static TArray<FVertexInstanceID> PossibleVertexInstances;
+								GetVertexInstancesInSameSoftEdgedGroup( VertexInstanceA.VertexID, VertexInstanceA.ConnectedPolygons[ 0 ], false, PossibleVertexInstances );
+
+								// If the smoothing group contains both vertex instances, they can be merged.
+								if( PossibleVertexInstances.Contains( VertexInstanceIDA ) && PossibleVertexInstances.Contains( VertexInstanceIDB ) )
+								{
+									// Change occurrences of VertexInstanceB for VertexInstanceA in VertexInstanceB's connected polygons.
+									// Note, this will cause VertexInstanceA's connected polygons list to be added to (at the end).
+									// This works because we are evaluating the number of connected polygons each time round the loop.
+									ReplaceVertexInstanceInPolygons( VertexInstanceIDB, VertexInstanceIDA, VertexInstanceB.ConnectedPolygons );
+
+									// This will also cause VertexInstanceB to be disconnected from all polygons.
+									// We mark the vertex instance for deletion here, but do not delete it until the end, as to do so would interrupt iterating through vertex instances.
+									check( VertexInstanceB.ConnectedPolygons.Num() == 0 );
+									VertexInstancesToDelete.Add( VertexInstanceIDB );
+								}
+							}
+						}
+					}
+				}
+
+				// Delete orphaned vertex instances
+				const bool bDeleteOrphanedVertices = false;
+				DeleteVertexInstances( VertexInstancesToDelete, bDeleteOrphanedVertices );
+			}
+
+			UE_LOG( LogEditableMesh, Verbose, TEXT( "* SetEdgeHardness returned" ) );
+		}
+
+		// Mark polygons adjacent to this edge as candidates for normal/tangent recalculation
+		PolygonsPendingNewTangentBasis.Append( Edge.ConnectedPolygons );
+	}
+}
+
+
+int32 UEditableMesh::GetEdgeCount() const
+{
+	return Edges.Num();
+}
+
+
+int32 UEditableMesh::GetEdgeArraySize() const
+{
+	return Edges.GetMaxIndex();
+}
+
+
+bool UEditableMesh::IsValidEdge( const FEdgeID EdgeID ) const
+{
+	return EdgeID.GetValue() >= 0 && EdgeID.GetValue() < Edges.GetMaxIndex() && Edges.IsAllocated( EdgeID.GetValue() );
+}
+
+
+FVertexID UEditableMesh::GetEdgeVertex( const FEdgeID EdgeID, const int32 EdgeVertexNumber ) const
+{
+	checkSlow( EdgeVertexNumber >= 0 && EdgeVertexNumber < 2 );
+	checkSlow( Edges.IsAllocated( EdgeID.GetValue() ) );
+	return Edges[ EdgeID.GetValue() ].VertexIDs[ EdgeVertexNumber ];
+}
+
+
+int32 UEditableMesh::GetEdgeConnectedPolygonCount( const FEdgeID EdgeID ) const
+{
+	checkSlow( Edges.IsAllocated( EdgeID.GetValue() ) );
+	const FMeshEdge& Edge = Edges[ EdgeID.GetValue() ];
+	return Edge.ConnectedPolygons.Num();
+}
+
+
+FPolygonID UEditableMesh::GetEdgeConnectedPolygon( const FEdgeID EdgeID, const int32 ConnectedPolygonNumber ) const
+{
+	checkSlow( Edges.IsAllocated( EdgeID.GetValue() ) );
+	const FMeshEdge& Edge = Edges[ EdgeID.GetValue() ];
+	const FPolygonID PolygonID = Edge.ConnectedPolygons[ ConnectedPolygonNumber ];
+	return PolygonID;
+}
+
+
+int32 UEditableMesh::GetPolygonGroupCount() const
+{
+	return PolygonGroups.Num();
+}
+
+
+int32 UEditableMesh::GetPolygonGroupArraySize() const
+{
+	return PolygonGroups.GetMaxIndex();
+}
+
+
+bool UEditableMesh::IsValidPolygonGroup( const FPolygonGroupID PolygonGroupID ) const
+{
+	return
+		PolygonGroupID.GetValue() >= 0 &&
+		PolygonGroupID.GetValue() < PolygonGroups.GetMaxIndex() &&
+		PolygonGroups.IsAllocated( PolygonGroupID.GetValue() );
+}
+
+
+int32 UEditableMesh::GetPolygonCountInGroup( const FPolygonGroupID PolygonGroupID ) const
+{
+	checkSlow( PolygonGroups.IsAllocated( PolygonGroupID.GetValue() ) );
+	return PolygonGroups[ PolygonGroupID.GetValue() ].Polygons.Num();
+}
+
+
+FPolygonID UEditableMesh::GetPolygonInGroup( const FPolygonGroupID PolygonGroupID, const int32 PolygonNumber ) const
+{
+	checkSlow( PolygonGroups.IsAllocated( PolygonGroupID.GetValue() ) );
+	const FMeshPolygonGroup& PolygonGroup = PolygonGroups[ PolygonGroupID.GetValue() ];
+	return FPolygonID( PolygonGroup.Polygons[ PolygonNumber ] );
+}
+
+
+int32 UEditableMesh::GetPolygonCount() const
+{
+	return Polygons.Num();
+}
+
+
+int32 UEditableMesh::GetPolygonArraySize() const
+{
+	return Polygons.GetMaxIndex();
+}
+
+
+bool UEditableMesh::IsValidPolygon( const FPolygonID PolygonID ) const
+{
+	return
+		PolygonID.GetValue() >= 0 &&
+		PolygonID.GetValue() < Polygons.GetMaxIndex() &&
+		Polygons.IsAllocated( PolygonID.GetValue() );
+}
+
+
+FPolygonGroupID UEditableMesh::GetGroupForPolygon( const FPolygonID PolygonID ) const
+{
+	checkSlow( Polygons.IsAllocated( PolygonID.GetValue() ) );
+	return Polygons[ PolygonID.GetValue() ].PolygonGroupID;
+}
+
+
+int32 UEditableMesh::GetPolygonPerimeterVertexCount( const FPolygonID PolygonID ) const
+{
+	checkSlow( Polygons.IsAllocated( PolygonID.GetValue() ) );
+	const FMeshPolygon& Polygon = Polygons[ PolygonID.GetValue() ];
+
+	return Polygon.PerimeterContour.VertexInstanceIDs.Num();
+}
+
+
+FVertexID UEditableMesh::GetPolygonPerimeterVertex( const FPolygonID PolygonID, const int32 PolygonVertexNumber ) const
+{
+	checkSlow( Polygons.IsAllocated( PolygonID.GetValue() ) );
+	const FMeshPolygon& Polygon = Polygons[ PolygonID.GetValue() ];
+	const FVertexInstanceID VertexInstanceID = Polygon.PerimeterContour.VertexInstanceIDs[ PolygonVertexNumber ];
+	const FMeshVertexInstance& VertexInstance = VertexInstances[ VertexInstanceID.GetValue() ];
+	return VertexInstance.VertexID;
+}
+
+
+FVertexInstanceID UEditableMesh::GetPolygonPerimeterVertexInstance( const FPolygonID PolygonID, const int32 PolygonVertexNumber ) const
+{
+	checkSlow( Polygons.IsAllocated( PolygonID.GetValue() ) );
+	const FMeshPolygon& Polygon = Polygons[ PolygonID.GetValue() ];
+	return Polygon.PerimeterContour.VertexInstanceIDs[ PolygonVertexNumber ];
+}
+
+
+FVector4 UEditableMesh::GetPolygonPerimeterVertexAttribute( const FPolygonID PolygonID, const int32 PolygonVertexNumber, const FName AttributeName, const int32 AttributeIndex ) const
+{
+	checkSlow( Polygons.IsAllocated( PolygonID.GetValue() ) );
+	const FMeshPolygon& Polygon = Polygons[ PolygonID.GetValue() ];
+	const FVertexInstanceID VertexInstanceID = Polygon.PerimeterContour.VertexInstanceIDs[ PolygonVertexNumber ];
+	return GetVertexInstanceAttribute( VertexInstanceID, AttributeName, AttributeIndex );
+}
+
+
+int32 UEditableMesh::GetPolygonHoleCount( const FPolygonID PolygonID ) const
+{
+	checkSlow( Polygons.IsAllocated( PolygonID.GetValue() ) );
+	const FMeshPolygon& Polygon = Polygons[ PolygonID.GetValue() ];
+
+	return Polygon.HoleContours.Num();
+}
+
+
+int32 UEditableMesh::GetPolygonHoleVertexCount( const FPolygonID PolygonID, const int32 HoleNumber ) const
+{
+	checkSlow( Polygons.IsAllocated( PolygonID.GetValue() ) );
+	const FMeshPolygon& Polygon = Polygons[ PolygonID.GetValue() ];
+	return Polygon.HoleContours[ HoleNumber ].VertexInstanceIDs.Num();
+}
+
+
+FVertexID UEditableMesh::GetPolygonHoleVertex( const FPolygonID PolygonID, const int32 HoleNumber, const int32 PolygonVertexNumber ) const
+{
+	checkSlow( Polygons.IsAllocated( PolygonID.GetValue() ) );
+	const FMeshPolygon& Polygon = Polygons[ PolygonID.GetValue() ];
+	const FVertexInstanceID VertexInstanceID = Polygon.HoleContours[ HoleNumber ].VertexInstanceIDs[ PolygonVertexNumber ];
+	const FMeshVertexInstance& VertexInstance = VertexInstances[ VertexInstanceID.GetValue() ];
+	return VertexInstance.VertexID;
+}
+
+
+FVertexInstanceID UEditableMesh::GetPolygonHoleVertexInstance( const FPolygonID PolygonID, const int32 HoleNumber, const int32 PolygonVertexNumber ) const
+{
+	checkSlow( Polygons.IsAllocated( PolygonID.GetValue() ) );
+	const FMeshPolygon& Polygon = Polygons[ PolygonID.GetValue() ];
+	return Polygon.HoleContours[ HoleNumber ].VertexInstanceIDs[ PolygonVertexNumber ];
+}
+
+
+FVector4 UEditableMesh::GetPolygonHoleVertexAttribute( const FPolygonID PolygonID, const int32 HoleNumber, const int32 PolygonVertexNumber, const FName AttributeName, const int32 AttributeIndex ) const
+{
+	checkSlow( Polygons.IsAllocated( PolygonID.GetValue() ) );
+	const FMeshPolygon& Polygon = Polygons[ PolygonID.GetValue() ];
+	const FVertexInstanceID VertexInstanceID = Polygon.HoleContours[ HoleNumber ].VertexInstanceIDs[ PolygonVertexNumber ];
+	return GetVertexInstanceAttribute( VertexInstanceID, AttributeName, AttributeIndex );
+}
+
+
+const TArray<FName>& UEditableMesh::GetValidVertexAttributes()
+{
+	static TArray<FName> ValidVertexAttributes;
+	if( ValidVertexAttributes.Num() == 0 )
+	{
+		ValidVertexAttributes.Add( UEditableMeshAttribute::VertexPosition() );
+		ValidVertexAttributes.Add( UEditableMeshAttribute::VertexCornerSharpness() );
+	}
+	return ValidVertexAttributes;
+}
+
+
+const TArray<FName>& UEditableMesh::GetValidVertexInstanceAttributes()
+{
+	static TArray<FName> ValidVertexInstanceAttributes;
+	if( ValidVertexInstanceAttributes.Num() == 0 )
+	{
+		ValidVertexInstanceAttributes.Add( UEditableMeshAttribute::VertexNormal() );
+		ValidVertexInstanceAttributes.Add( UEditableMeshAttribute::VertexTangent() );
+		ValidVertexInstanceAttributes.Add( UEditableMeshAttribute::VertexBinormalSign() );
+		ValidVertexInstanceAttributes.Add( UEditableMeshAttribute::VertexTextureCoordinate() );
+		ValidVertexInstanceAttributes.Add( UEditableMeshAttribute::VertexColor() );
+	}
+	return ValidVertexInstanceAttributes;
+}
+
+
+const TArray<FName>& UEditableMesh::GetValidEdgeAttributes()
+{
+	static TArray<FName> ValidEdgeAttributes;
+	if( ValidEdgeAttributes.Num() == 0 )
+	{
+		ValidEdgeAttributes.Add( UEditableMeshAttribute::EdgeIsHard() );
+		ValidEdgeAttributes.Add( UEditableMeshAttribute::EdgeCreaseSharpness() );
+	}
+	return ValidEdgeAttributes;
+}
+
+
+const TArray<FName>& UEditableMesh::GetValidPolygonAttributes()
+{
+	static TArray<FName> ValidPolygonAttributes;
+	if( ValidPolygonAttributes.Num() == 0 )
+	{
+		ValidPolygonAttributes.Add( UEditableMeshAttribute::PolygonNormal() );
+		ValidPolygonAttributes.Add( UEditableMeshAttribute::PolygonCenter() );
+	}
+	return ValidPolygonAttributes;
+}
+
+
+int32 UEditableMesh::GetMaxAttributeIndex( const FName AttributeName ) const
+{
+	if( AttributeName == UEditableMeshAttribute::VertexTextureCoordinate() )
+	{
+		return GetTextureCoordinateCount();
+	}
+
+	return 1;
+}
+
+
+FPolygonGroupID UEditableMesh::GetFirstValidPolygonGroup() const
+{
+	FPolygonGroupID FirstValidPolygonGroupID = FPolygonGroupID::Invalid;
+
+	if( PolygonGroups.Num() > 0 )
+	{
+		FirstValidPolygonGroupID = FPolygonGroupID( PolygonGroups.CreateConstIterator().GetIndex() );
+	}
+
+	return FirstValidPolygonGroupID;
+}
+
+
+int32 UEditableMesh::GetTextureCoordinateCount() const
+{
+	return this->TextureCoordinateCount;
+}
+
+
+int32 UEditableMesh::GetSubdivisionCount() const
+{
+	return this->SubdivisionCount;
+}
+
+
+bool UEditableMesh::IsPreviewingSubdivisions() const
+{
+	return GetSubdivisionCount() > 0;		// @todo mesheditor subdiv: Make optional even when subdivisions are active
+}
+
+
+void UEditableMesh::GetVertexConnectedEdges( const FVertexID VertexID, TArray<FEdgeID>& OutConnectedEdgeIDs ) const
+{
+	OutConnectedEdgeIDs.Reset();
+
+	const int32 ConnectedEdgeCount = GetVertexConnectedEdgeCount( VertexID );
+	OutConnectedEdgeIDs.Reserve( ConnectedEdgeCount );
+	for( int32 EdgeNumber = 0; EdgeNumber < ConnectedEdgeCount; ++EdgeNumber )
+	{
+		const FEdgeID ConnectedEdgeID = GetVertexConnectedEdge( VertexID, EdgeNumber );
+		OutConnectedEdgeIDs.Add( ConnectedEdgeID );
+	}
+}
+
+
+void UEditableMesh::GetVertexConnectedPolygons( const FVertexID VertexID, TArray<FPolygonID>& OutConnectedPolygonIDs ) const
+{
+	OutConnectedPolygonIDs.Reset();
+
+	checkSlow( Vertices.IsAllocated( VertexID.GetValue() ) );
+	const FMeshVertex& Vertex = Vertices[ VertexID.GetValue() ];
+	for( const FVertexInstanceID VertexInstanceID : Vertex.VertexInstanceIDs )
+	{
+		checkSlow( VertexInstances.IsAllocated( VertexInstanceID.GetValue() ) );
+		const FMeshVertexInstance& VertexInstance = VertexInstances[ VertexInstanceID.GetValue() ];
+		OutConnectedPolygonIDs.Append( VertexInstance.ConnectedPolygons );
+	}
+}
+
+
+void UEditableMesh::GetVertexInstanceConnectedPolygons( const FVertexInstanceID VertexInstanceID, TArray<FPolygonID>& OutConnectedPolygonIDs ) const
+{
+	checkSlow( VertexInstances.IsAllocated( VertexInstanceID.GetValue() ) );
+	const FMeshVertexInstance& VertexInstance = VertexInstances[ VertexInstanceID.GetValue() ];
+
+	OutConnectedPolygonIDs = VertexInstance.ConnectedPolygons;
+}
+
+
+int32 UEditableMesh::GetVertexInstanceConnectedPolygonCount( const FVertexInstanceID VertexInstanceID ) const
+{
+	checkSlow( VertexInstances.IsAllocated( VertexInstanceID.GetValue() ) );
+	const FMeshVertexInstance& VertexInstance = VertexInstances[ VertexInstanceID.GetValue() ];
+	return VertexInstance.ConnectedPolygons.Num();
+}
+
+
+void UEditableMesh::GetVertexAdjacentVertices( const FVertexID VertexID, TArray<FVertexID>& OutAdjacentVertexIDs ) const
+{
+	OutAdjacentVertexIDs.Reset();
+
+	const FMeshVertex& Vertex = Vertices[ VertexID.GetValue() ];
+	for( const FEdgeID EdgeID : Vertex.ConnectedEdgeIDs )
+	{
+		checkSlow( Edges.IsAllocated( EdgeID.GetValue() ) );
+		const FMeshEdge& Edge = Edges[ EdgeID.GetValue() ];
+		OutAdjacentVertexIDs.Add( Edge.VertexIDs[ 0 ] == VertexID ? Edge.VertexIDs[ 1 ] : Edge.VertexIDs[ 0 ] );
+	}
+}
+
+
+void UEditableMesh::GetEdgeVertices( const FEdgeID EdgeID, FVertexID& OutEdgeVertexID0, FVertexID& OutEdgeVertexID1 ) const
+{
+	checkSlow( Edges.IsAllocated( EdgeID.GetValue() ) );
+	const FMeshEdge& Edge = Edges[ EdgeID.GetValue() ];
+	OutEdgeVertexID0 = Edge.VertexIDs[ 0 ];
+	OutEdgeVertexID1 = Edge.VertexIDs[ 1 ];
+}
+
+
+void UEditableMesh::GetEdgeConnectedPolygons( const FEdgeID EdgeID, TArray<FPolygonID>& OutConnectedPolygonIDs ) const
+{
+	checkSlow( Edges.IsAllocated( EdgeID.GetValue() ) );
+	const FMeshEdge& Edge = Edges[ EdgeID.GetValue() ];
+	OutConnectedPolygonIDs = Edge.ConnectedPolygons;
+}
+
+
+FEdgeID UEditableMesh::GetEdgeThatConnectsVertices( const FVertexID VertexID0, const FVertexID VertexID1 ) const
+{
+	FEdgeID EdgeID = FEdgeID::Invalid;
+
+	const int32 Vertex0ConnectedEdgeCount = GetVertexConnectedEdgeCount( VertexID0 );
+	for( int32 Vertex0ConnectedEdgeNumber = 0; Vertex0ConnectedEdgeNumber < Vertex0ConnectedEdgeCount; ++Vertex0ConnectedEdgeNumber )
+	{
+		const FEdgeID Vertex0ConnectedEdge = GetVertexConnectedEdge( VertexID0, Vertex0ConnectedEdgeNumber );
+
+		FVertexID EdgeVertexID0, EdgeVertexID1;
+		GetEdgeVertices( Vertex0ConnectedEdge, /* Out */ EdgeVertexID0, /* Out */ EdgeVertexID1 );
+
+		if( ( EdgeVertexID0 == VertexID0 && EdgeVertexID1 == VertexID1 ) ||
+			( EdgeVertexID0 == VertexID1 && EdgeVertexID1 == VertexID0 ) )
+		{
+			EdgeID = Vertex0ConnectedEdge;
+			break;
+		}
+	}
+
+	check( EdgeID != FEdgeID::Invalid );
+	return EdgeID;
+}
+
+
+void UEditableMesh::GetEdgeLoopElements( const FEdgeID EdgeID, TArray<FEdgeID>& EdgeLoopIDs ) const
+{
+	EdgeLoopIDs.Reset();
+
+	// Maintain a list of unique edge IDs which form the loop
+	static TSet<FEdgeID> EdgeIDs;
+	EdgeIDs.Reset();
+	
+	// Maintain a stack of edges to be processed, in lieu of recursion.
+	// We also store which vertex of the edge has already been processed (so we don't retrace our steps when processing stack items).
+	static TArray<TTuple<FEdgeID, FVertexID>> EdgeStack;
+	EdgeStack.Reset();
+	EdgeStack.Push( MakeTuple( EdgeID, FVertexID::Invalid ) );
+
+	// Process edge IDs on the stack
+	while( EdgeStack.Num() > 0 )
+	{
+		const TTuple<FEdgeID, FVertexID> CurrentItem = EdgeStack.Pop();
+		const FEdgeID CurrentEdgeID = CurrentItem.Get<0>();
+		const FVertexID FromVertexID = CurrentItem.Get<1>();
+		EdgeIDs.Add( CurrentEdgeID );
+
+		// Get the edge so we can see which polygons are connected to it.
+		// When continuing the loop, the criterion is that new edges must share no polygons with this edge,
+		// i.e. they are the other side of a perpendicular edge.
+		checkSlow( Edges.IsAllocated( CurrentEdgeID.GetValue() ) );
+		const FMeshEdge& Edge = Edges[ CurrentEdgeID.GetValue() ];
+
+		// Now look for edges connected to each end of this edge
+		for( const FVertexID ConnectedVertexID : Edge.VertexIDs )
+		{
+			// If we have already processed this vertex, skip it
+			if( ConnectedVertexID == FromVertexID )
+			{
+				continue;
+			}
+
+			// This is the candidate edge ID which continues the loop beyond the vertex being processed
+			FEdgeID AdjacentEdgeID = FEdgeID::Invalid;
+
+			// Iterate through all edges connected to this vertex
+			const FMeshVertex& Vertex = Vertices[ ConnectedVertexID.GetValue() ];
+			for( const FEdgeID ConnectedEdgeID : Vertex.ConnectedEdgeIDs )
+			{
+				// If this edge hasn't been added to the loop...
+				if( !EdgeIDs.Contains( ConnectedEdgeID ) )
+				{
+					// ...see if it shares any polygons with the original edge (intersection operation)
+					bool bIsCandidateEdge = true;
+					checkSlow( Edges.IsAllocated( ConnectedEdgeID.GetValue() ) );
+					const FMeshEdge& ConnectedEdge = Edges[ ConnectedEdgeID.GetValue() ];
+					for( const FPolygonID ConnectedPolygonID : ConnectedEdge.ConnectedPolygons )
+					{
+						if( Edge.ConnectedPolygons.Contains( ConnectedPolygonID ) )
+						{
+							bIsCandidateEdge = false;
+							break;
+						}
+					}
+
+					// We have found an edge connected to this vertex which doesn't share any polys with the original edge
+					if( bIsCandidateEdge )
+					{
+						if( AdjacentEdgeID == FEdgeID::Invalid )
+						{
+							// If it's the first such edge which meets the criteria, remember it
+							AdjacentEdgeID = ConnectedEdgeID;
+						}
+						else
+						{
+							// If we already have a possible edge, stop the loop here; we don't allow splits in the loop if there is more than one candidate
+							AdjacentEdgeID = FEdgeID::Invalid;
+							break;
+						}
+					}
+				}
+			}
+
+			if( AdjacentEdgeID != FEdgeID::Invalid )
+			{
+				EdgeStack.Push( MakeTuple( AdjacentEdgeID, ConnectedVertexID ) );
+			}
+		}
+	}
+
+	for( const FEdgeID EdgeLoopID : EdgeIDs )
+	{
+		EdgeLoopIDs.Add( EdgeLoopID );
+	}
+}
+
+
+int32 UEditableMesh::GetPolygonPerimeterEdgeCount( const FPolygonID PolygonID ) const
+{
+	// All polygons have the same number of edges as they do vertices
+	return GetPolygonPerimeterVertexCount( PolygonID );
+}
+
+
+int32 UEditableMesh::GetPolygonHoleEdgeCount( const FPolygonID PolygonID, const int32 HoleNumber ) const
+{
+	// All polygons have the same number of edges as they do vertices
+	return GetPolygonHoleVertexCount( PolygonID, HoleNumber );
+}
+
+
+void UEditableMesh::GetPolygonPerimeterVertices( const FPolygonID PolygonID, TArray<FVertexID>& OutPolygonPerimeterVertexIDs ) const
+{
+	checkSlow( Polygons.IsAllocated( PolygonID.GetValue() ) );
+	const FMeshPolygon& Polygon = Polygons[ PolygonID.GetValue() ];
+	OutPolygonPerimeterVertexIDs.SetNumUninitialized( Polygon.PerimeterContour.VertexInstanceIDs.Num(), false );
+
+	int32 Index = 0;
+	for( const FVertexInstanceID VertexInstanceID : Polygon.PerimeterContour.VertexInstanceIDs )
+	{
+		const FMeshVertexInstance& VertexInstance = VertexInstances[ VertexInstanceID.GetValue() ];
+		OutPolygonPerimeterVertexIDs[ Index ] = VertexInstance.VertexID;
+		Index++;
+	}
+}
+
+
+void UEditableMesh::GetPolygonPerimeterVertexInstances( const FPolygonID PolygonID, TArray<FVertexInstanceID>& OutPolygonPerimeterVertexInstanceIDs ) const
+{
+	checkSlow( Polygons.IsAllocated( PolygonID.GetValue() ) );
+	const FMeshPolygon& Polygon = Polygons[ PolygonID.GetValue() ];
+	OutPolygonPerimeterVertexInstanceIDs = Polygon.PerimeterContour.VertexInstanceIDs;
+}
+
+
+void UEditableMesh::GetPolygonHoleVertices( const FPolygonID PolygonID, const int32 HoleNumber, TArray<FVertexID>& OutHoleVertexIDs ) const
+{
+	checkSlow( Polygons.IsAllocated( PolygonID.GetValue() ) );
+	const FMeshPolygon& Polygon = Polygons[ PolygonID.GetValue() ];
+	const FMeshPolygonContour& Contour = Polygon.HoleContours[ HoleNumber ];
+
+	OutHoleVertexIDs.SetNumUninitialized( Contour.VertexInstanceIDs.Num(), false );
+
+	int32 Index = 0;
+	for( const FVertexInstanceID VertexInstanceID : Contour.VertexInstanceIDs )
+	{
+		const FMeshVertexInstance& VertexInstance = VertexInstances[ VertexInstanceID.GetValue() ];
+		OutHoleVertexIDs[ Index ] = VertexInstance.VertexID;
+		Index++;
+	}
+}
+
+
+void UEditableMesh::GetPolygonHoleVertexInstances( const FPolygonID PolygonID, const int32 HoleNumber, TArray<FVertexInstanceID>& OutHoleVertexInstanceIDs ) const
+{
+	checkSlow( Polygons.IsAllocated( PolygonID.GetValue() ) );
+	const FMeshPolygon& Polygon = Polygons[ PolygonID.GetValue() ];
+	const FMeshPolygonContour& Contour = Polygon.HoleContours[ HoleNumber ];
+
+	OutHoleVertexInstanceIDs = Contour.VertexInstanceIDs;
+}
+
+
+FEdgeID UEditableMesh::GetPolygonContourEdge( const FMeshPolygonContour& Contour, const int32 ContourEdgeNumber, bool& bOutEdgeWindingIsReversed ) const
+{
+	const int32 NumPolygonContourEdges = Contour.VertexInstanceIDs.Num();
+
+	const FVertexInstanceID VertexInstanceID = Contour.VertexInstanceIDs[ ContourEdgeNumber ];
+	const FVertexInstanceID NextVertexInstanceID = Contour.VertexInstanceIDs[ ( ContourEdgeNumber + 1 ) % NumPolygonContourEdges ];
+
+	checkSlow( VertexInstances.IsAllocated( VertexInstanceID.GetValue() ) );
+	checkSlow( VertexInstances.IsAllocated( NextVertexInstanceID.GetValue() ) );
+	const FVertexID VertexID = VertexInstances[ VertexInstanceID.GetValue() ].VertexID;
+	const FVertexID NextVertexID = VertexInstances[ NextVertexInstanceID.GetValue() ].VertexID;
+
+	return GetVertexPairEdge( VertexID, NextVertexID, bOutEdgeWindingIsReversed );
+}
+
+
+FEdgeID UEditableMesh::GetPolygonPerimeterEdge( const FPolygonID PolygonID, const int32 PerimeterEdgeNumber, bool& bOutEdgeWindingIsReversedForPolygon ) const
+{
+	checkSlow( Polygons.IsAllocated( PolygonID.GetValue() ) );
+	const FMeshPolygon& Polygon = Polygons[ PolygonID.GetValue() ];
+	return GetPolygonContourEdge( Polygon.PerimeterContour, PerimeterEdgeNumber, bOutEdgeWindingIsReversedForPolygon );
+}
+
+
+FEdgeID UEditableMesh::GetVertexPairEdge( const FVertexID StartVertexID, const FVertexID EndVertexID, bool& bOutEdgeWindingIsReversed ) const
+{
+	bOutEdgeWindingIsReversed = false;
+
+	const FMeshVertex& StartVertex = Vertices[ StartVertexID.GetValue() ];
+	for( const FEdgeID VertexConnectedEdgeID : StartVertex.ConnectedEdgeIDs )
+	{
+		const FMeshEdge& Edge = Edges[ VertexConnectedEdgeID.GetValue() ];
+
+		// Try the edge's first vertex.  Does it point to us?
+		FVertexID OtherEdgeVertexID = Edge.VertexIDs[ 0 ];
+		bOutEdgeWindingIsReversed = false;
+		if( OtherEdgeVertexID == StartVertexID )
+		{
+			// Must be the the other one
+			OtherEdgeVertexID = Edge.VertexIDs[ 1 ];
+			bOutEdgeWindingIsReversed = true;
+		}
+
+		if( OtherEdgeVertexID == EndVertexID )
+		{
+			// We found the edge!
+			return VertexConnectedEdgeID;
+		}
+	}
+
+	return FEdgeID::Invalid;
+}
+
+
+FEdgeID UEditableMesh::GetPolygonHoleEdge( const FPolygonID PolygonID, const int32 HoleNumber, const int32 HoleEdgeNumber ) const
+{
+	checkSlow( Polygons.IsAllocated( PolygonID.GetValue() ) );
+	const FMeshPolygon& Polygon = Polygons[ PolygonID.GetValue() ];
+	const FMeshPolygonContour& Contour = Polygon.HoleContours[ HoleNumber ];
+
+	bool bOutEdgeWindingIsReversed;
+	return GetPolygonContourEdge( Contour, HoleEdgeNumber, bOutEdgeWindingIsReversed );
+}
+
+
+void UEditableMesh::GetPolygonContourEdges( const FMeshPolygonContour& Contour, TArray<FEdgeID>& OutPolygonContourEdgeIDs ) const
+{
+	const int32 NumContourEdges = Contour.VertexInstanceIDs.Num();
+	OutPolygonContourEdgeIDs.SetNumUninitialized( NumContourEdges, false );
+
+	for( int32 EdgeNumber = 0; EdgeNumber < NumContourEdges; ++EdgeNumber )
+	{
+		bool bOutEdgeWindingIsReversed;
+		OutPolygonContourEdgeIDs[ EdgeNumber ] = GetPolygonContourEdge( Contour, EdgeNumber, bOutEdgeWindingIsReversed );
+	}
+}
+
+
+void UEditableMesh::GetPolygonPerimeterEdges( const FPolygonID PolygonID, TArray<FEdgeID>& OutPolygonPerimeterEdgeIDs ) const
+{
+	checkSlow( Polygons.IsAllocated( PolygonID.GetValue() ) );
+	const FMeshPolygon& Polygon = Polygons[ PolygonID.GetValue() ];
+	GetPolygonContourEdges( Polygon.PerimeterContour, OutPolygonPerimeterEdgeIDs );
+}
+
+
+void UEditableMesh::GetPolygonHoleEdges( const FPolygonID PolygonID, const int32 HoleNumber, TArray<FEdgeID>& OutHoleEdgeIDs ) const
+{
+	checkSlow( Polygons.IsAllocated( PolygonID.GetValue() ) );
+	const FMeshPolygon& Polygon = Polygons[ PolygonID.GetValue() ];
+	const FMeshPolygonContour& Contour = Polygon.HoleContours[ HoleNumber ];
+	GetPolygonContourEdges( Contour, OutHoleEdgeIDs );
+}
+
+
+int32 UEditableMesh::GetPolygonTriangulatedTriangleCount( const FPolygonID PolygonID ) const
+{
+	checkSlow( Polygons.IsAllocated( PolygonID.GetValue() ) );
+	const FMeshPolygon& Polygon = Polygons[ PolygonID.GetValue() ];
+
+	return Polygon.Triangles.Num();
+}
+
+
+FMeshTriangle UEditableMesh::GetPolygonTriangulatedTriangle( const FPolygonID PolygonID, int32 PolygonTriangleNumber ) const
+{
+	checkSlow( Polygons.IsAllocated( PolygonID.GetValue() ) );
+	const FMeshPolygon& Polygon = Polygons[ PolygonID.GetValue() ];
+
+	// @todo mesheditor: how to make this safe so that BPs don't crash when given an out-of-range index.
+	// (same goes for just about all of these methods)
+	return Polygon.Triangles[ PolygonTriangleNumber ];
+}
+
+
+FVector UEditableMesh::GetPolygonTriangulatedTriangleVertexPosition( const FPolygonID PolygonID, const int32 PolygonTriangleNumber, const int32 TriangleVertexNumber ) const
+{
+	checkSlow( TriangleVertexNumber >= 0 && TriangleVertexNumber < 3 );
+
+	checkSlow( Polygons.IsAllocated( PolygonID.GetValue() ) );
+	const FMeshPolygon& Polygon = Polygons[ PolygonID.GetValue() ];
+
+	const FVertexInstanceID VertexInstanceID = Polygon.Triangles[ PolygonTriangleNumber ].GetVertexInstanceID( TriangleVertexNumber );
+	checkSlow( VertexInstances.IsAllocated( VertexInstanceID.GetValue() ) );
+	const FVertexID VertexID = VertexInstances[ VertexInstanceID.GetValue() ].VertexID;
+
+	checkSlow( Vertices.IsAllocated( VertexID.GetValue() ) );
+	return Vertices[ VertexID.GetValue() ].VertexPosition;
+}
+
+
+void UEditableMesh::GetPolygonAdjacentPolygons( const FPolygonID PolygonID, TArray<FPolygonID>& OutAdjacentPolygons ) const
+{
+	OutAdjacentPolygons.Reset();
+
+	static TArray<FEdgeID> PolygonPerimeterEdges;
+	GetPolygonPerimeterEdges( PolygonID, PolygonPerimeterEdges );
+
+	for( const FEdgeID EdgeID : PolygonPerimeterEdges )
+	{
+		static TArray<FPolygonID> EdgeConnectedPolygons;
+		GetEdgeConnectedPolygons( EdgeID, EdgeConnectedPolygons );
+
+		for( const FPolygonID EdgeConnectedPolygon : EdgeConnectedPolygons )
+		{
+			if( EdgeConnectedPolygon != PolygonID )
+			{
+				OutAdjacentPolygons.AddUnique( EdgeConnectedPolygon );
+			}
+		}
+	}
+}
+
+
+FBox UEditableMesh::ComputeBoundingBox() const
+{
+	FBox BoundingBox;
+	BoundingBox.Init();
+
+	for( TSparseArray<FMeshVertex>::TConstIterator It( Vertices ); It; ++It )
+	{
+		BoundingBox += It->VertexPosition;
+	}
+
+	return BoundingBox;
+}
+
+
+FBoxSphereBounds UEditableMesh::ComputeBoundingBoxAndSphere() const
+{
+	const FBox BoundingBox = ComputeBoundingBox();
+
+	FBoxSphereBounds BoundingBoxAndSphere;
+	BoundingBox.GetCenterAndExtents( /* Out */ BoundingBoxAndSphere.Origin, /* Out */ BoundingBoxAndSphere.BoxExtent );
+
+	// Calculate the bounding sphere, using the center of the bounding box as the origin.
+	BoundingBoxAndSphere.SphereRadius = 0.0f;
+
+	const int32 VertexArraySize = GetVertexArraySize();
+	for( int32 VertexNumber = 0; VertexNumber < VertexArraySize; ++VertexNumber )
+	{
+		const FVertexID VertexID( VertexNumber );
+		if( IsValidVertex( VertexID ) )
+		{
+			const FVector VertexPosition = GetVertexAttribute( VertexID, UEditableMeshAttribute::VertexPosition(), 0 );
+
+			BoundingBoxAndSphere.SphereRadius = FMath::Max( ( VertexPosition - BoundingBoxAndSphere.Origin ).Size(), BoundingBoxAndSphere.SphereRadius );
+		}
+	}
+
+	return BoundingBoxAndSphere;
+}
+
+
+FVector UEditableMesh::ComputePolygonCenter( const FPolygonID PolygonID ) const
+{
+	FVector Centroid = FVector::ZeroVector;
+
+	static TArray<FVertexID> PerimeterVertexIDs;
+	GetPolygonPerimeterVertices( PolygonID, /* Out */ PerimeterVertexIDs );
+
+	for( const FVertexID VertexID : PerimeterVertexIDs )
+	{
+		const FVector Position = GetVertexAttribute( VertexID, UEditableMeshAttribute::VertexPosition(), 0 );
+
+		Centroid += Position;
+	}
+
+	return Centroid / (float)PerimeterVertexIDs.Num();
+}
+
+
+FPlane UEditableMesh::ComputePolygonPlane( const FPolygonID PolygonID ) const
+{
+	// NOTE: This polygon plane computation code is partially based on the implementation of "Newell's method" from Real-Time 
+	//       Collision Detection by Christer Ericson, published by Morgan Kaufmann Publishers, (c) 2005 Elsevier Inc
+
+	// @todo mesheditor perf: For polygons that are just triangles, use a cross product to get the normal fast!
+	// @todo mesheditor perf: We could skip computing the plane distance when we only need the normal
+	// @todo mesheditor perf: We could cache these computed polygon normals; or just use the normal of the first three vertices' triangle if it is satisfactory in all cases
+	// @todo mesheditor: For non-planar polygons, the result can vary. Ideally this should use the actual polygon triangulation as opposed to the arbitrary triangulation used here.
+
+	FVector Centroid = FVector::ZeroVector;
+	FVector Normal = FVector::ZeroVector;
+
+	static TArray<FVertexID> PerimeterVertexIDs;
+	GetPolygonPerimeterVertices( PolygonID, /* Out */ PerimeterVertexIDs );
+
+	// Use 'Newell's Method' to compute a robust 'best fit' plane from the vertices of this polygon
+	for( int32 VertexNumberI = PerimeterVertexIDs.Num() - 1, VertexNumberJ = 0; VertexNumberJ < PerimeterVertexIDs.Num(); VertexNumberI = VertexNumberJ, VertexNumberJ++ )
+	{
+		const FVertexID VertexIDI = PerimeterVertexIDs[ VertexNumberI ];
+		const FVector PositionI = GetVertexAttribute( VertexIDI, UEditableMeshAttribute::VertexPosition(), 0 );
+
+		const FVertexID VertexIDJ = PerimeterVertexIDs[ VertexNumberJ ];
+		const FVector PositionJ = GetVertexAttribute( VertexIDJ, UEditableMeshAttribute::VertexPosition(), 0 );
+
+		Centroid += PositionJ;
+
+		Normal.X += ( PositionJ.Y - PositionI.Y ) * ( PositionI.Z + PositionJ.Z );
+		Normal.Y += ( PositionJ.Z - PositionI.Z ) * ( PositionI.X + PositionJ.X );
+		Normal.Z += ( PositionJ.X - PositionI.X ) * ( PositionI.Y + PositionJ.Y );
+	}
+
+	Normal.Normalize();
+
+	// Construct a plane from the normal and centroid
+	return FPlane( Normal, FVector::DotProduct( Centroid, Normal ) / (float)PerimeterVertexIDs.Num() );
+}
+
+
+FVector UEditableMesh::ComputePolygonNormal( const FPolygonID PolygonID ) const
+{
+	// @todo mesheditor: Polygon normals are now computed and cached when changes are made to a polygon.
+	// In theory, we can just return that cached value, but we need to check that there is nothing which relies on the value being correct before
+	// the cache is updated at the end of a modification.
+	const FPlane PolygonPlane = ComputePolygonPlane( PolygonID );
+	const FVector PolygonNormal( PolygonPlane.X, PolygonPlane.Y, PolygonPlane.Z );
+	return PolygonNormal;
+}
+
+
+void UEditableMesh::RefreshOpenSubdiv()
+{
+	OsdTopologyRefiner.Reset();
+
+	if( SubdivisionCount > 0 )
+	{
+		// @todo mesheditor subdiv perf: Ideally we give our topology data straight to OSD rather than have it build it from this simple 
+		// set of parameters.  This will be much more efficient!  In order to do this, we need to specialize TopologyRefinerFactory 
+		// for our mesh (http://graphics.pixar.com/opensubdiv/docs/far_overview.html)
+		OpenSubdiv::Far::TopologyDescriptor OsdTopologyDescriptor;
+		{
+			const int32 VertexArraySize = GetVertexArraySize();
+			OsdTopologyDescriptor.numVertices = VertexArraySize;
+			OsdTopologyDescriptor.numFaces = GetPolygonCount();
+
+			// NOTE: OpenSubdiv likes weights to be between 0.0 and 10.0, so we'll account for that here
+			const float OpenSubdivCreaseWeightMultiplier = 10.0f;
+
+			{
+				// Subdivision corner weights
+				{
+					OsdCornerVertexIndices.Reset();
+					OsdCornerWeights.Reset();
+					for( int32 VertexNumber = 0; VertexNumber < VertexArraySize; ++VertexNumber )
+					{
+						const FVertexID VertexID( VertexNumber );
+						if( IsValidVertex( VertexID ) )
+						{
+							const float VertexCornerSharpness = GetVertexAttribute( VertexID, UEditableMeshAttribute::VertexCornerSharpness(), 0 ).X;
+							if( VertexCornerSharpness > SMALL_NUMBER )
+							{
+								// This vertex is (at least partially) a subdivision corner
+								OsdCornerVertexIndices.Add( VertexNumber );
+								OsdCornerWeights.Add( OpenSubdivCreaseWeightMultiplier * VertexCornerSharpness );
+							}
+						}
+					}
+				}
+
+				// Edge creases
+				{
+					const int32 EdgeArraySize = GetEdgeArraySize();
+
+					OsdCreaseVertexIndexPairs.Reset();
+					OsdCreaseWeights.Reset();
+					for( int32 EdgeNumber = 0; EdgeNumber < EdgeArraySize; ++EdgeNumber )
+					{
+						const FEdgeID EdgeID( EdgeNumber );
+						if( IsValidEdge( EdgeID ) )
+						{
+							const float EdgeCreaseSharpness = GetEdgeAttribute( EdgeID, UEditableMeshAttribute::EdgeCreaseSharpness(), 0 ).X;
+							if( EdgeCreaseSharpness > SMALL_NUMBER )
+							{
+								// This edge is (at least partially) creased
+								FVertexID EdgeVertexID0, EdgeVertexID1;
+								GetEdgeVertices( EdgeID, /* Out */ EdgeVertexID0, /* Out */ EdgeVertexID1 );
+
+								OsdCreaseVertexIndexPairs.Add( EdgeVertexID0.GetValue() );
+								OsdCreaseVertexIndexPairs.Add( EdgeVertexID1.GetValue() );
+								OsdCreaseWeights.Add( OpenSubdivCreaseWeightMultiplier * EdgeCreaseSharpness );
+							}
+						}
+					}
+				}
+
+				OsdNumVerticesPerFace.SetNum( OsdTopologyDescriptor.numFaces, false );
+				OsdVertexIndicesPerFace.Reset();
+				OsdFVarIndicesPerFace.Reset();
+
+				int32 NextOsdFaceIndex = 0;
+
+				const int32 PolygonCount = GetPolygonArraySize();
+				for( int32 PolygonIndex = 0; PolygonIndex < PolygonCount; ++PolygonIndex )
+				{
+					const FPolygonID PolygonID( PolygonIndex );
+					if( IsValidPolygon( PolygonID ) )
+					{
+						static TArray<FVertexID> PerimeterVertexIDs;
+						GetPolygonPerimeterVertices( PolygonID, /* Out */ PerimeterVertexIDs );
+
+						const int32 PerimeterVertexCount = PerimeterVertexIDs.Num();
+						OsdNumVerticesPerFace[ NextOsdFaceIndex++ ] = PerimeterVertexCount;
+
+						for( int32 PerimeterVertexNumber = 0; PerimeterVertexNumber < PerimeterVertexCount; ++PerimeterVertexNumber )
+						{
+							const FVertexID PerimeterVertexID = PerimeterVertexIDs[ PerimeterVertexNumber ];
+
+							OsdVertexIndicesPerFace.Add( PerimeterVertexID.GetValue() );
+							OsdFVarIndicesPerFace.Add( OsdFVarIndicesPerFace.Num() );
+						}
+					}
+				}
+
+				check( NextOsdFaceIndex == OsdNumVerticesPerFace.Num()  );
+				check( OsdVertexIndicesPerFace.Num() == OsdFVarIndicesPerFace.Num() );
+			}
+
+			{
+				const int32 TotalFVarChannels = 1;
+				OsdFVarChannels.SetNum( TotalFVarChannels, false );
+				for( int32 FVarChannelNumber = 0; FVarChannelNumber < OsdFVarChannels.Num(); ++FVarChannelNumber )
+				{
+					FOsdFVarChannel& OsdFVarChannel = OsdFVarChannels[ FVarChannelNumber ];
+
+					// @todo mesheditor subdiv: For now, we'll assuming unique face-varying data for every polygon vertex.  Ideally
+					// if we were able to share rendering vertices, we'd want to share face-varying data too.  OpenSubdiv allows
+					// you to pass a separate index array for every channel of face-varying data, but that's probably overkill
+					// for us.  Instead we just need to reflect what we do with rendering vertices.
+					OsdFVarChannel.ValueCount = OsdFVarIndicesPerFace.Num();
+					OsdFVarChannel.ValueIndices = OsdFVarIndicesPerFace.GetData();
+				}
+			}
+
+			OsdTopologyDescriptor.numVertsPerFace = OsdNumVerticesPerFace.GetData();
+			OsdTopologyDescriptor.vertIndicesPerFace = OsdVertexIndicesPerFace.GetData();
+
+			OsdTopologyDescriptor.numCreases = OsdCreaseWeights.Num();
+			OsdTopologyDescriptor.creaseVertexIndexPairs = OsdCreaseVertexIndexPairs.GetData();
+			OsdTopologyDescriptor.creaseWeights = OsdCreaseWeights.GetData();
+
+			OsdTopologyDescriptor.numCorners = OsdCornerWeights.Num();
+			OsdTopologyDescriptor.cornerVertexIndices = OsdCornerVertexIndices.GetData();
+			OsdTopologyDescriptor.cornerWeights = OsdCornerWeights.GetData();
+
+			OsdTopologyDescriptor.numHoles = 0;	// @todo mesheditor holes
+			OsdTopologyDescriptor.holeIndices = nullptr;
+
+			OsdTopologyDescriptor.isLeftHanded = true;
+
+			// Face-varying vertex data.  This maps to our GetPolygonVertexAttribute() calls.
+			OsdTopologyDescriptor.numFVarChannels = OsdFVarChannels.Num();
+			OsdTopologyDescriptor.fvarChannels = reinterpret_cast<OpenSubdiv::v3_2_0::Far::TopologyDescriptor::FVarChannel*>( OsdFVarChannels.GetData() );
+		}
+
+		// We always want Catmull-Clark subdivisions
+		const OpenSubdiv::Sdc::SchemeType OsdSchemeType = OpenSubdiv::Sdc::SCHEME_CATMARK;
+
+		OpenSubdiv::Sdc::Options OsdSdcOptions;
+		{
+			// @todo mesheditor subdiv: Tweak these settings
+			OsdSdcOptions.SetVtxBoundaryInterpolation( OpenSubdiv::Sdc::Options::VTX_BOUNDARY_EDGE_ONLY );
+			OsdSdcOptions.SetFVarLinearInterpolation( OpenSubdiv::Sdc::Options::FVAR_LINEAR_ALL );
+			OsdSdcOptions.SetCreasingMethod( OpenSubdiv::Sdc::Options::CREASE_UNIFORM );
+			OsdSdcOptions.SetTriangleSubdivision( OpenSubdiv::Sdc::Options::TRI_SUB_CATMARK );
+		}
+
+		OpenSubdiv::Far::TopologyRefinerFactory<OpenSubdiv::Far::TopologyDescriptor>::Options OsdTopologyRefinerOptions( OsdSchemeType, OsdSdcOptions );
+
+		this->OsdTopologyRefiner = MakeShareable(
+			OpenSubdiv::Far::TopologyRefinerFactory<OpenSubdiv::Far::TopologyDescriptor>::Create(
+				OsdTopologyDescriptor,
+				OsdTopologyRefinerOptions ) );
+
+		OpenSubdiv::Far::TopologyRefiner::UniformOptions OsdUniformOptions( SubdivisionCount );
+		{
+			// @todo mesheditor subdiv: Should we order child vertices from faces instead of child vertices from vertices first?
+			OsdUniformOptions.orderVerticesFromFacesFirst = false;
+
+			// NOTE: In order for face-varying data to work, OpenSubdiv requires 'fullTopologyInLastLevel' to be enabled.
+			OsdUniformOptions.fullTopologyInLastLevel = true;
+		}
+
+
+		this->OsdTopologyRefiner->RefineUniform( OsdUniformOptions );
+	}
+
+	GenerateOpenSubdivLimitSurfaceData();
+}
+
+
+const FSubdivisionLimitData& UEditableMesh::GetSubdivisionLimitData() const
+{
+	return SubdivisionLimitData;
+}
+
+
+void UEditableMesh::GenerateOpenSubdivLimitSurfaceData()
+{
+	SubdivisionLimitData = FSubdivisionLimitData();
+
+	if( SubdivisionCount > 0 && ensure( this->OsdTopologyRefiner.IsValid() ) )
+	{
+		// Create an OpenSubdiv 'primvar refiner'.  This guy allows us to interpolate data between vertices on a subdivision level.
+		const OpenSubdiv::Far::PrimvarRefiner OsdPrimvarRefiner( *OsdTopologyRefiner );
+
+		struct FOsdVector
+		{
+			void Clear( void* UnusedPtr = nullptr )
+			{
+				Position = FVector( 0, 0, 0 );
+			}
+
+			void AddWithWeight( FOsdVector const& SourceVertexPosition, float Weight )
+			{
+				Position += SourceVertexPosition.Position * Weight;
+			}
+
+			FVector Position;
+		};
+
+		struct FOsdFVarVertexData
+		{
+			void Clear()
+			{
+				TextureCoordinates[0] = TextureCoordinates[1] = FVector2D( 0.0f, 0.0f );
+				VertexColor = FLinearColor( 0.0f, 0.0f, 0.0f, 0.0f );
+			}
+
+			void AddWithWeight( FOsdFVarVertexData const& SourceVertex, float Weight )
+			{
+				TextureCoordinates[ 0 ] += SourceVertex.TextureCoordinates[ 0 ] * Weight;
+				TextureCoordinates[ 1 ] += SourceVertex.TextureCoordinates[ 1 ] * Weight;
+				VertexColor += SourceVertex.VertexColor * Weight;
+			}
+
+			FVector2D TextureCoordinates[2];	// @todo mesheditor subdiv perf: Only two UV sets supported for now (just to avoid heap allocs for dynamic array)
+			FLinearColor VertexColor;
+		};
+
+		const int32 PolygonGroupCount = GetPolygonGroupCount();
+
+		// Get the limit surface subdivision level from OpenSubdiv
+		const OpenSubdiv::Far::TopologyLevel& OsdLimitLevel = OsdTopologyRefiner->GetLevel( SubdivisionCount );
+
+		const int32 LimitVertexCount = OsdLimitLevel.GetNumVertices();
+		const int32 LimitFaceCount = OsdLimitLevel.GetNumFaces();
+
+		static TArray<FVector> LimitXGradients;
+		LimitXGradients.Reset();
+		static TArray<FVector> LimitYGradients;
+		LimitYGradients.Reset();
+
+		// Grab all of the vertex data and put them in separate contiguous arrays for OpenSubdiv
+		static TArray<FVector> VertexPositions;
+		static TArray<FOsdFVarVertexData> FVarVertexDatas;
+		static TArray<int32> FirstPolygonNumberForPolygonGroups;
+		{
+			// Vertex positions
+			{
+				// NOTE: We're including an entry for all vertices, even vertices that aren't referenced by any triangles (due to our
+				//       sparse array optimization.)
+				const int32 VertexArraySize = GetVertexArraySize();
+				VertexPositions.SetNum( VertexArraySize, false );
+				for( int32 VertexNumber = 0; VertexNumber < VertexArraySize; ++VertexNumber )
+				{
+					const FVertexID VertexID( VertexNumber );
+					if( IsValidVertex( VertexID ) )
+					{
+						VertexPositions[ VertexNumber ] = GetVertexAttribute( VertexID, UEditableMeshAttribute::VertexPosition(), 0 );
+					}
+					else
+					{
+						// Vertex isn't used, but we'll include a zero'd entry so that our indices still match up.
+						VertexPositions[ VertexNumber ] = FVector::ZeroVector;
+					}
+				}
+			}
+
+			// Texture coordinates (per polygon vertex)
+			FirstPolygonNumberForPolygonGroups.Reset();
+
+			{
+				FVarVertexDatas.Reset();
+				FVarVertexDatas.Reserve( OsdFVarIndicesPerFace.Num() );
+
+				const int32 PolygonGroupArraySize = GetPolygonGroupArraySize();
+
+				int32 NumPolygonsSoFar = 0;
+				for( int32 PolygonGroupNumber = 0; PolygonGroupNumber < PolygonGroupArraySize; ++PolygonGroupNumber )
+				{
+					const FPolygonGroupID PolygonGroupID( PolygonGroupNumber );
+					if( IsValidPolygonGroup( PolygonGroupID ) )
+					{
+						FirstPolygonNumberForPolygonGroups.Add( NumPolygonsSoFar );
+						const int32 PolygonCountInGroup = GetPolygonCountInGroup( PolygonGroupID );
+						NumPolygonsSoFar += PolygonCountInGroup;
+
+						for( int32 PolygonGroupPolygonNumber = 0; PolygonGroupPolygonNumber < PolygonCountInGroup; ++PolygonGroupPolygonNumber )
+						{
+							const FPolygonID PolygonID( GetPolygonInGroup( PolygonGroupID, PolygonGroupPolygonNumber ) );
+
+							const int32 PerimeterVertexCount = GetPolygonPerimeterVertexCount( PolygonID );
+							for( int32 PerimeterVertexNumber = 0; PerimeterVertexNumber < PerimeterVertexCount; ++PerimeterVertexNumber )
+							{
+								FOsdFVarVertexData& FVarVertexData = *new( FVarVertexDatas ) FOsdFVarVertexData();
+								FVarVertexData.TextureCoordinates[ 0 ] = TextureCoordinateCount > 0 ? FVector2D( GetPolygonPerimeterVertexAttribute( PolygonID, PerimeterVertexNumber, UEditableMeshAttribute::VertexTextureCoordinate(), 0 ) ) : FVector2D::ZeroVector;
+								FVarVertexData.TextureCoordinates[ 1 ] = TextureCoordinateCount > 1 ? FVector2D( GetPolygonPerimeterVertexAttribute( PolygonID, PerimeterVertexNumber, UEditableMeshAttribute::VertexTextureCoordinate(), 1 ) ) : FVector2D::ZeroVector;
+								FVarVertexData.VertexColor = FLinearColor( GetPolygonPerimeterVertexAttribute( PolygonID, PerimeterVertexNumber, UEditableMeshAttribute::VertexColor(), 0 ) );
+							}
+						}
+					}
+				}
+
+				check( FVarVertexDatas.Num() == OsdVertexIndicesPerFace.Num() );
+			}
+		}
+
+		// @todo mesheditor subdiv debug
+// 		for( int32 VertexIndex = 0; VertexIndex < FVarVertexDatas.Num(); ++VertexIndex )
+// 		{
+// 			const FOsdFVarVertexData& FVarVertexData = FVarVertexDatas[ VertexIndex ];
+// 			GWarn->Logf( TEXT( "FVar%i, U:%0.2f, V:%0.2f" ), VertexIndex, FVarVertexData.TextureCoordinates[0].X, FVarVertexData.TextureCoordinates[0].Y );
+// 		}
+
+		static TArray<FOsdFVarVertexData> LimitFVarVertexDatas;
+		LimitFVarVertexDatas.Reset();
+
+
+		// Start with the base cage geometry, and refine the geometry until we get to the limit surface
+		{
+			// NOTE: The OsdVertexPositions list might contains vertices that aren't actually referenced by any polygons (due to our
+			//       sparse array optimization.)  That's OK though.
+			{
+				int32 NextScratchBufferIndex = 0;
+				static TArray<FVector> ScratchVertexPositions[ 2 ];
+				ScratchVertexPositions[ 0 ].Reset();
+				ScratchVertexPositions[ 1 ].Reset();
+
+				for( int32 RefinementLevel = 1; RefinementLevel <= SubdivisionCount; ++RefinementLevel )
+				{
+					const OpenSubdiv::Far::TopologyLevel& OsdLevel = OsdTopologyRefiner->GetLevel( RefinementLevel );
+
+					// For the last refinement level, we'll copy positions straight to our output buffer (to avoid having to copy the data later.)
+					// For earlier levels, we'll ping-pong between scratch buffers.
+					const TArray<FVector>& SourceVertexPositions = ( RefinementLevel == 1 ) ? VertexPositions : ScratchVertexPositions[ !NextScratchBufferIndex ];
+					check( SourceVertexPositions.Num() == OsdTopologyRefiner->GetLevel( RefinementLevel - 1 ).GetNumVertices() );
+
+					TArray<FVector>& DestVertexPositions = ScratchVertexPositions[ NextScratchBufferIndex ];
+					DestVertexPositions.SetNumUninitialized( OsdLevel.GetNumVertices(), false );
+
+					FOsdVector* OsdDestVertexPositions = reinterpret_cast<FOsdVector*>( DestVertexPositions.GetData() );
+					OsdPrimvarRefiner.Interpolate( 
+						RefinementLevel, 
+						reinterpret_cast<const FOsdVector*>( SourceVertexPositions.GetData() ),
+						OsdDestVertexPositions );
+
+					NextScratchBufferIndex = !NextScratchBufferIndex;
+				}
+
+				// @todo mesheditor subdiv perf: We should probably be using stencils for faster updates (unless topology is changing every frame too)
+
+				// We've generated interpolated positions for the most fine subdivision level, but now we need to compute the positions
+				// on the limit surface.  While doing this, we also compute gradients at every vertex for either surface axis
+				{
+					const FOsdVector* OsdSourceLimitPositions = reinterpret_cast<FOsdVector*>( ScratchVertexPositions[ !NextScratchBufferIndex ].GetData() );
+					SubdivisionLimitData.VertexPositions.SetNumUninitialized( LimitVertexCount );
+					FOsdVector* OsdDestLimitPositions = reinterpret_cast<FOsdVector*>( SubdivisionLimitData.VertexPositions.GetData() );
+
+					LimitXGradients.SetNumUninitialized( LimitVertexCount );
+					FOsdVector* OsdDestLimitXGradients = reinterpret_cast<FOsdVector*>( LimitXGradients.GetData() );
+
+					LimitYGradients.SetNumUninitialized( LimitVertexCount );
+					FOsdVector* OsdDestLimitYGradients = reinterpret_cast<FOsdVector*>( LimitYGradients.GetData() );
+
+					OsdPrimvarRefiner.Limit( OsdSourceLimitPositions, /* Out */ OsdDestLimitPositions, /* Out */ OsdDestLimitXGradients, /* Out */ OsdDestLimitYGradients );
+
+					if( EditableMesh::InterpolatePositionsToLimit->GetInt() == 0 )
+					{
+						SubdivisionLimitData.VertexPositions = ScratchVertexPositions[ !NextScratchBufferIndex ];
+					}
+				}
+
+				check( LimitVertexCount == SubdivisionLimitData.VertexPositions.Num() );
+			}
+
+			{
+				const int32 FVarChannelNumber = 0;
+
+				static TArray<FOsdFVarVertexData> ScratchFVarVertexDatas[ 2 ];
+				ScratchFVarVertexDatas[ 0 ].Reset();
+				ScratchFVarVertexDatas[ 1 ].Reset();
+
+				int32 NextScratchBufferIndex = 0;
+				for( int32 RefinementLevel = 1; RefinementLevel <= SubdivisionCount; ++RefinementLevel )
+				{
+					const OpenSubdiv::Far::TopologyLevel& OsdLevel = OsdTopologyRefiner->GetLevel( RefinementLevel );
+
+					const TArray<FOsdFVarVertexData>& SourceFVarVertexDatas = ( RefinementLevel == 1 ) ? FVarVertexDatas : ScratchFVarVertexDatas[ !NextScratchBufferIndex ];
+					check( SourceFVarVertexDatas.Num() == OsdTopologyRefiner->GetLevel( RefinementLevel - 1 ).GetNumFVarValues( FVarChannelNumber ) );
+
+					TArray<FOsdFVarVertexData>& DestFVarVertexDatas = ScratchFVarVertexDatas[ NextScratchBufferIndex ];
+					DestFVarVertexDatas.SetNumUninitialized( OsdLevel.GetNumFVarValues( FVarChannelNumber ), false );
+
+					FOsdFVarVertexData* DestFVarVertexDatasPtr = DestFVarVertexDatas.GetData();
+					OsdPrimvarRefiner.InterpolateFaceVarying(
+						RefinementLevel,
+						SourceFVarVertexDatas.GetData(),
+						DestFVarVertexDatasPtr,
+						FVarChannelNumber );
+
+					NextScratchBufferIndex = !NextScratchBufferIndex;
+				}
+	
+				if( EditableMesh::InterpolateFVarsToLimit->GetInt() != 0 )
+				{
+					LimitFVarVertexDatas.SetNumUninitialized( OsdLimitLevel.GetNumFVarValues( FVarChannelNumber ), false );
+					FOsdFVarVertexData* DestFVarVertexDatasPtr = LimitFVarVertexDatas.GetData();
+					OsdPrimvarRefiner.LimitFaceVarying(	// @todo mesheditor subdiv fvars: The OSD tutorials don't bother doing a Limit pass for UVs/Colors...
+						ScratchFVarVertexDatas[ !NextScratchBufferIndex ].GetData(),
+						DestFVarVertexDatasPtr,
+						FVarChannelNumber );
+				}
+				else
+				{
+					LimitFVarVertexDatas = ScratchFVarVertexDatas[ !NextScratchBufferIndex ];
+				}
+			}
+		}
+
+		SubdivisionLimitData.Sections.SetNum( PolygonGroupCount, false );
+
+		for( int32 LimitFaceNumber = 0; LimitFaceNumber < LimitFaceCount; ++LimitFaceNumber )
+		{
+			const OpenSubdiv::Far::ConstIndexArray& OsdLimitFaceVertices = OsdLimitLevel.GetFaceVertices( LimitFaceNumber );
+			const int32 FaceVertexCount = OsdLimitFaceVertices.size();
+			check( FaceVertexCount == 4 );	// We're always expecting quads as the result of a Catmull-Clark subdivision
+
+			// Find the parent face in our original control mesh for this subdivided quad.  We'll use this to
+			// determine which section the face belongs to
+			// @todo mesheditor subdiv perf: We can use InterpolateFaceUniform() to push section indices down to the subdivided faces. More memory, might be slower (copies), but avoids iterating here.
+			int32 QuadSectionNumber = 0;
+			{
+				int32 CurrentFaceNumber = LimitFaceNumber;
+				for( int32 SubdivisionLevel = SubdivisionCount; SubdivisionLevel > 0; --SubdivisionLevel )
+				{
+					const OpenSubdiv::Far::TopologyLevel& OsdLevel = OsdTopologyRefiner->GetLevel( SubdivisionLevel );
+					const int32 ParentLevelFace = OsdLevel.GetFaceParentFace( CurrentFaceNumber );
+					CurrentFaceNumber = ParentLevelFace;
+				}
+				int32 BaseCageFaceNumber = CurrentFaceNumber;
+
+				for( int32 PolygonGroupNumber = PolygonGroupCount - 1; PolygonGroupNumber >= 0; --PolygonGroupNumber)
+				{
+					if( BaseCageFaceNumber >= FirstPolygonNumberForPolygonGroups[ PolygonGroupNumber ] )
+					{
+						QuadSectionNumber = PolygonGroupNumber;
+						break;
+					}
+				}
+			}
+
+			const int32 FVarChannelNumber = 0;
+			const OpenSubdiv::Far::ConstIndexArray& OsdLimitFaceFVarValues = OsdLimitLevel.GetFaceFVarValues( LimitFaceNumber, FVarChannelNumber );
+			check( OsdLimitFaceFVarValues.size() == 4 );	// Expecting quads
+
+			FSubdivisionLimitSection& SubdivisionSection = SubdivisionLimitData.Sections[ QuadSectionNumber ];
+			FSubdividedQuad& SubdividedQuad = *new( SubdivisionSection.SubdividedQuads ) FSubdividedQuad();
+			for( int32 FaceVertexNumber = 0; FaceVertexNumber < FaceVertexCount; ++FaceVertexNumber )
+			{
+				FSubdividedQuadVertex& QuadVertex = SubdividedQuad.AccessQuadVertex( FaceVertexNumber );
+
+				QuadVertex.VertexPositionIndex = OsdLimitFaceVertices[ FaceVertexNumber ];
+
+				const int OsdLimitFaceFVarIndex = OsdLimitFaceFVarValues[ FaceVertexNumber ];
+				const FOsdFVarVertexData& FVarVertexData = LimitFVarVertexDatas[ OsdLimitFaceFVarIndex ];
+
+				QuadVertex.TextureCoordinate0 = FVarVertexData.TextureCoordinates[ 0 ];
+				QuadVertex.TextureCoordinate1 = FVarVertexData.TextureCoordinates[ 1 ];
+
+				QuadVertex.VertexColor = FVarVertexData.VertexColor.ToFColor( true );
+
+				QuadVertex.VertexNormal = FVector::CrossProduct( LimitXGradients[ QuadVertex.VertexPositionIndex ].GetSafeNormal(), LimitYGradients[ QuadVertex.VertexPositionIndex ].GetSafeNormal() );
+				
+				// NOTE: Tangents will be computed separately, below
+			}
+		}
+
+		// Compute normal and tangent vectors for each quad vertex, taking into account the texture coordinates
+		for( int32 SectionNumber = 0; SectionNumber < SubdivisionLimitData.Sections.Num(); ++SectionNumber )
+		{
+			struct FMikkUserData
+			{
+				FSubdivisionLimitData* LimitData;
+				int32 SectionNumber;
+
+				FMikkUserData( FSubdivisionLimitData* InitLimitData, const int32 InitSectionNumber )
+					: LimitData( InitLimitData ),
+						SectionNumber( InitSectionNumber )
+				{
+				}
+			} MikkUserData( &SubdivisionLimitData, SectionNumber );
+
+			struct Local
+			{
+				static int MikkGetNumFaces( const SMikkTSpaceContext* Context )
+				{
+					const FMikkUserData& UserData = *static_cast<FMikkUserData*>( Context->m_pUserData );
+					return UserData.LimitData->Sections[ UserData.SectionNumber ].SubdividedQuads.Num();
+				}
+
+				static int MikkGetNumVertsOfFace( const SMikkTSpaceContext* Context, const int MikkFaceIndex )
+				{
+					// Always quads
+					return 4;
+				}
+
+				static void MikkGetPosition( const SMikkTSpaceContext* Context, float OutPosition[ 3 ], const int MikkFaceIndex, const int MikkVertexIndex )
+				{
+					const FMikkUserData& UserData = *static_cast<FMikkUserData*>( Context->m_pUserData );
+					const FSubdividedQuadVertex& QuadVertex = UserData.LimitData->Sections[ UserData.SectionNumber ].SubdividedQuads[ MikkFaceIndex ].GetQuadVertex( MikkVertexIndex );
+					const FVector VertexPosition = UserData.LimitData->VertexPositions[ QuadVertex.VertexPositionIndex ];
+
+					OutPosition[ 0 ] = VertexPosition.X;
+					OutPosition[ 1 ] = VertexPosition.Y;
+					OutPosition[ 2 ] = VertexPosition.Z;
+				}
+
+				static void MikkGetNormal( const SMikkTSpaceContext* Context, float OutNormal[ 3 ], const int MikkFaceIndex, const int MikkVertexIndex )
+				{
+					const FMikkUserData& UserData = *static_cast<FMikkUserData*>( Context->m_pUserData );
+					const FSubdividedQuadVertex& QuadVertex = UserData.LimitData->Sections[ UserData.SectionNumber ].SubdividedQuads[ MikkFaceIndex ].GetQuadVertex( MikkVertexIndex );
+
+					OutNormal[ 0 ] = QuadVertex.VertexNormal.X;
+					OutNormal[ 1 ] = QuadVertex.VertexNormal.Y;
+					OutNormal[ 2 ] = QuadVertex.VertexNormal.Z;
+				}
+
+				static void MikkGetTexCoord( const SMikkTSpaceContext* Context, float OutUV[ 2 ], const int MikkFaceIndex, const int MikkVertexIndex )
+				{
+					const FMikkUserData& UserData = *static_cast<FMikkUserData*>( Context->m_pUserData );
+					const FSubdividedQuadVertex& QuadVertex = UserData.LimitData->Sections[ UserData.SectionNumber ].SubdividedQuads[ MikkFaceIndex ].GetQuadVertex( MikkVertexIndex );
+
+					// @todo mesheditor: Support using a custom texture coordinate index for tangent space generation?
+					OutUV[ 0 ] = QuadVertex.TextureCoordinate0.X;
+					OutUV[ 1 ] = QuadVertex.TextureCoordinate0.Y;
+				}
+
+				static void MikkSetTSpaceBasic( const SMikkTSpaceContext* Context, const float Tangent[ 3 ], const float BitangentSign, const int MikkFaceIndex, const int MikkVertexIndex )
+				{
+					FMikkUserData& UserData = *static_cast<FMikkUserData*>( Context->m_pUserData );
+					FSubdividedQuadVertex& QuadVertex = UserData.LimitData->Sections[ UserData.SectionNumber ].SubdividedQuads[ MikkFaceIndex ].AccessQuadVertex( MikkVertexIndex );
+
+					QuadVertex.VertexTangent = FVector( Tangent[ 0 ], Tangent[ 1 ], Tangent[ 2 ] );
+					QuadVertex.VertexBinormalSign = BitangentSign;
+				}
+			};
+
+			SMikkTSpaceInterface MikkTInterface;
+			{
+				MikkTInterface.m_getNumFaces = Local::MikkGetNumFaces;
+				MikkTInterface.m_getNumVerticesOfFace = Local::MikkGetNumVertsOfFace;
+				MikkTInterface.m_getPosition = Local::MikkGetPosition;
+				MikkTInterface.m_getNormal = Local::MikkGetNormal;
+				MikkTInterface.m_getTexCoord = Local::MikkGetTexCoord;
+
+				MikkTInterface.m_setTSpaceBasic = Local::MikkSetTSpaceBasic;
+				MikkTInterface.m_setTSpace = nullptr;
+			}
+
+			SMikkTSpaceContext MikkTContext;
+			{
+				MikkTContext.m_pInterface = &MikkTInterface;
+				MikkTContext.m_pUserData = (void*)( &MikkUserData );
+				MikkTContext.m_bIgnoreDegenerates = true;	// No degenerates in our list
+			}
+
+			// Now we'll ask MikkTSpace to actually generate the tangents
+			genTangSpaceDefault( &MikkTContext );
+		}
+
+
+		// Generate our edge information for the subdivided mesh.  We'll also figure out which subdivided edges have a 
+		// counterpart on the base cage mesh, so tools can display this information to the user.
+		{
+			const int32 LimitEdgeCount = OsdLimitLevel.GetNumEdges();
+			for( int32 LimitEdgeNumber = 0; LimitEdgeNumber < LimitEdgeCount; ++LimitEdgeNumber )
+			{
+				const OpenSubdiv::Far::ConstIndexArray& OsdLimitEdgeVertices = OsdLimitLevel.GetEdgeVertices( LimitEdgeNumber );
+				const int32 EdgeVertexCount = OsdLimitEdgeVertices.size();
+				check( EdgeVertexCount == 2 );	// Edges always connect two vertices
+
+				FSubdividedWireEdge& SubdividedWireEdge = *new( SubdivisionLimitData.SubdividedWireEdges ) FSubdividedWireEdge();
+
+				SubdividedWireEdge.EdgeVertex0PositionIndex = OsdLimitEdgeVertices[ 0 ];
+				SubdividedWireEdge.EdgeVertex1PositionIndex = OsdLimitEdgeVertices[ 1 ];
+
+				// Default to not highlighting this edge as a base cage counterpart.  We'll actually figure this out below.
+				SubdividedWireEdge.CounterpartEdgeID = FEdgeID::Invalid;
+			}
+
+			{
+				static TSet<int32> BaseCageEdgeSet;
+				BaseCageEdgeSet.Reset();
+
+				const OpenSubdiv::Far::TopologyLevel& OsdBaseCageLevel = OsdTopologyRefiner->GetLevel( 0 );
+				const int32 BaseCageFaceCount = OsdBaseCageLevel.GetNumFaces();
+				for( int32 BaseCageFaceNumber = 0; BaseCageFaceNumber < BaseCageFaceCount; ++BaseCageFaceNumber )
+				{
+					const OpenSubdiv::Far::ConstIndexArray& OsdBaseCageFaceEdges = OsdBaseCageLevel.GetFaceEdges( BaseCageFaceNumber );
+					for( int32 FaceEdgeNumber = 0; FaceEdgeNumber < OsdBaseCageFaceEdges.size(); ++FaceEdgeNumber )
+					{
+						const int32 BaseCageEdgeIndex = OsdBaseCageFaceEdges[ FaceEdgeNumber ];
+						bool bIsAlreadyInSet = false;
+						BaseCageEdgeSet.Add( BaseCageEdgeIndex, /* Out */ &bIsAlreadyInSet );
+						if( !bIsAlreadyInSet )
+						{
+							// Find our original edge ID for each of the OpenSubdiv base cage edges
+							const OpenSubdiv::Far::ConstIndexArray& OsdEdgeVertices = OsdBaseCageLevel.GetEdgeVertices( BaseCageEdgeIndex );
+							check( OsdEdgeVertices.size() == 2 ); // Edges always connect two vertices
+																	// Figure out which edge goes with these vertices
+							const FEdgeID BaseCageEdgeID = GetEdgeThatConnectsVertices( FVertexID( OsdEdgeVertices[ 0 ] ), FVertexID( OsdEdgeVertices[ 1 ] ) );
+
+							// Go through and determine the limit child edges of all of the original base cage edges by drilling down through
+							// the subdivision hierarchy
+							int32 NextScratchBufferIndex = 0;
+							static TArray<int32> ScratchChildEdges[ 2 ];
+							ScratchChildEdges[ 0 ].Reset();
+							ScratchChildEdges[ 1 ].Reset();
+
+							// Fill in our source buffer with the starting edge
+							ScratchChildEdges[ NextScratchBufferIndex ].Add( BaseCageEdgeIndex );
+							NextScratchBufferIndex = !NextScratchBufferIndex;
+
+							for( int32 RefinementLevel = 0; RefinementLevel < SubdivisionCount; ++RefinementLevel )
+							{
+								const OpenSubdiv::Far::TopologyLevel& OsdLevel = OsdTopologyRefiner->GetLevel( RefinementLevel );
+
+								const TArray<int32>& SourceChildEdges = ScratchChildEdges[ !NextScratchBufferIndex ];
+
+								TArray<int32>& DestChildEdges = ScratchChildEdges[ NextScratchBufferIndex ];
+								DestChildEdges.Reset();
+
+								for( int32 SourceEdgeNumber = 0; SourceEdgeNumber < SourceChildEdges.Num(); ++SourceEdgeNumber )
+								{
+									const OpenSubdiv::Far::ConstIndexArray& OsdChildEdges = OsdLevel.GetEdgeChildEdges( SourceChildEdges[ SourceEdgeNumber ] );
+									for( int32 ChildEdgeNumber = 0; ChildEdgeNumber < OsdChildEdges.size(); ++ChildEdgeNumber )
+									{
+										DestChildEdges.Add( OsdChildEdges[ ChildEdgeNumber ] );
+									}
+								}
+
+								NextScratchBufferIndex = !NextScratchBufferIndex;
+							}
+
+							// No go back and update our subdivided wire edges, marking the edges that we determined were descendants of the base cage edges
+							const TArray<int32>& BaseCageCounterpartEdgesAtLimit = ScratchChildEdges[ !NextScratchBufferIndex ];
+							for( const int32 CounterpartEdgeAtLimit : BaseCageCounterpartEdgesAtLimit )
+							{
+								check( CounterpartEdgeAtLimit < SubdivisionLimitData.SubdividedWireEdges.Num() );
+								SubdivisionLimitData.SubdividedWireEdges[ CounterpartEdgeAtLimit ].CounterpartEdgeID = BaseCageEdgeID;
+							}
+						}
+					}
+				}
+			}
+		}
+	}
+}
+
+
+void UEditableMesh::RetriangulatePolygons()
+{
+	// Perform triangulation directly into mesh polygons
+	for( const FPolygonID PolygonID : PolygonsPendingTriangulation )
+	{
+		FMeshPolygon& Polygon = Polygons[ PolygonID.GetValue() ];
+		ComputePolygonTriangulation( PolygonID, Polygon.Triangles );
+	}
+
+	for( UEditableMeshAdapter* Adapter : Adapters )
+	{
+		Adapter->OnRetriangulatePolygons( this, PolygonsPendingTriangulation.Array() );
+	}
+}
+
+
+void UEditableMesh::ComputePolygonTriangulation( const FPolygonID PolygonID, TArray<FMeshTriangle>& OutTriangles ) const
+{
+	// NOTE: This polygon triangulation code is partially based on the ear cutting algorithm described on
+	//       page 497 of the book "Real-time Collision Detection", published in 2005.
+
+	struct Local
+	{
+		// Returns true if the triangle formed by the specified three positions has a normal that is facing the opposite direction of the reference normal
+		static inline bool IsTriangleFlipped( const FVector ReferenceNormal, const FVector VertexPositionA, const FVector VertexPositionB, const FVector VertexPositionC )
+		{
+			const FVector TriangleNormal = FVector::CrossProduct(
+				VertexPositionC - VertexPositionA,
+				VertexPositionB - VertexPositionA ).GetSafeNormal();
+			return ( FVector::DotProduct( ReferenceNormal, TriangleNormal ) <= 0.0f );
+		}
+
+	};
+
+
+	OutTriangles.Reset();
+
+	// @todo mesheditor holes: Does not support triangles with holes yet!
+	// @todo mesheditor: Perhaps should always attempt to triangulate by splitting polygons along the shortest edge, for better determinism.
+
+	const FMeshPolygon& Polygon = Polygons[ PolygonID.GetValue() ];
+	const TArray<FVertexInstanceID>& PolygonVertexInstanceIDs = Polygon.PerimeterContour.VertexInstanceIDs;
+
+	// Polygon must have at least three vertices/edges
+	const int32 PolygonVertexCount = PolygonVertexInstanceIDs.Num();
+	check( PolygonVertexCount >= 3 );
+
+	// First figure out the polygon normal.  We need this to determine which triangles are convex, so that
+	// we can figure out which ears to clip
+	const FVector PolygonNormal = ComputePolygonNormal( PolygonID );
+
+	// Make a simple linked list array of the previous and next vertex numbers, for each vertex number
+	// in the polygon.  This will just save us having to iterate later on.
+	static TArray<int32> PrevVertexNumbers, NextVertexNumbers;
+	static TArray<FVector> VertexPositions;
+	{
+		PrevVertexNumbers.SetNumUninitialized( PolygonVertexCount, false );
+		NextVertexNumbers.SetNumUninitialized( PolygonVertexCount, false );
+		VertexPositions.SetNumUninitialized( PolygonVertexCount, false );
+
+		for( int32 VertexNumber = 0; VertexNumber < PolygonVertexCount; ++VertexNumber )
+		{
+			PrevVertexNumbers[ VertexNumber ] = VertexNumber - 1;
+			NextVertexNumbers[ VertexNumber ] = VertexNumber + 1;
+
+			VertexPositions[ VertexNumber ] = GetVertexInstanceAttribute( PolygonVertexInstanceIDs[ VertexNumber ], UEditableMeshAttribute::VertexPosition(), 0 );
+		}
+		PrevVertexNumbers[ 0 ] = PolygonVertexCount - 1;
+		NextVertexNumbers[ PolygonVertexCount - 1 ] = 0;
+	}
+
+	int32 EarVertexNumber = 0;
+	int32 EarTestCount = 0;
+	for( int32 RemainingVertexCount = PolygonVertexCount; RemainingVertexCount >= 3; )
+	{
+		bool bIsEar = true;
+
+		// If we're down to only a triangle, just treat it as an ear.  Also, if we've tried every possible candidate
+		// vertex looking for an ear, go ahead and just treat the current vertex as an ear.  This can happen when 
+		// vertices are colinear or other degenerate cases.
+		if( RemainingVertexCount > 3 && EarTestCount < RemainingVertexCount )
+		{
+			const FVector PrevVertexPosition = VertexPositions[ PrevVertexNumbers[ EarVertexNumber ] ];
+			const FVector EarVertexPosition = VertexPositions[ EarVertexNumber ];
+			const FVector NextVertexPosition = VertexPositions[ NextVertexNumbers[ EarVertexNumber ] ];
+
+			// Figure out whether the potential ear triangle is facing the same direction as the polygon
+			// itself.  If it's facing the opposite direction, then we're dealing with a concave triangle
+			// and we'll skip it for now.
+			if( !Local::IsTriangleFlipped(
+					PolygonNormal,
+					PrevVertexPosition,
+					EarVertexPosition,
+					NextVertexPosition ) )
+			{
+				int32 TestVertexNumber = NextVertexNumbers[ NextVertexNumbers[ EarVertexNumber ] ];
+
+				do 
+				{
+					// Test every other remaining vertex to make sure that it doesn't lie inside our potential ear
+					// triangle.  If we find a vertex that's inside the triangle, then it cannot actually be an ear.
+					const FVector TestVertexPosition = VertexPositions[ TestVertexNumber ];
+					if( FGeomTools::PointInTriangle(
+							PrevVertexPosition,
+							EarVertexPosition,
+							NextVertexPosition,
+							TestVertexPosition,
+							SMALL_NUMBER ) )
+					{
+						bIsEar = false;
+						break;
+					}
+
+					TestVertexNumber = NextVertexNumbers[ TestVertexNumber ];
+				} 
+				while( TestVertexNumber != PrevVertexNumbers[ EarVertexNumber ] );
+			}
+			else
+			{
+				bIsEar = false;
+			}
+		}
+
+		if( bIsEar )
+		{
+			// OK, we found an ear!  Let's save this triangle in our output buffer.
+			{
+				OutTriangles.Emplace();
+				FMeshTriangle& Triangle = OutTriangles.Last();
+
+				Triangle.SetVertexInstanceID( 0, PolygonVertexInstanceIDs[ PrevVertexNumbers[ EarVertexNumber ] ] );
+				Triangle.SetVertexInstanceID( 1, PolygonVertexInstanceIDs[ EarVertexNumber ] );
+				Triangle.SetVertexInstanceID( 2, PolygonVertexInstanceIDs[ NextVertexNumbers[ EarVertexNumber ] ] );
+			}
+
+			// Update our linked list.  We're effectively cutting off the ear by pointing the ear vertex's neighbors to
+			// point at their next sequential neighbor, and reducing the remaining vertex count by one.
+			{
+				NextVertexNumbers[ PrevVertexNumbers[ EarVertexNumber ] ] = NextVertexNumbers[ EarVertexNumber ];
+				PrevVertexNumbers[ NextVertexNumbers[ EarVertexNumber ] ] = PrevVertexNumbers[ EarVertexNumber ];
+				--RemainingVertexCount;
+			}
+
+			// Move on to the previous vertex in the list, now that this vertex was cut
+			EarVertexNumber = PrevVertexNumbers[ EarVertexNumber ];
+
+			EarTestCount = 0;
+		}
+		else
+		{
+			// The vertex is not the ear vertex, because it formed a triangle that either had a normal which pointed in the opposite direction
+			// of the polygon, or at least one of the other polygon vertices was found to be inside the triangle.  Move on to the next vertex.
+			EarVertexNumber = NextVertexNumbers[ EarVertexNumber ];
+
+			// Keep track of how many ear vertices we've tested, so that if we exhaust all remaining vertices, we can
+			// fall back to clipping the triangle and adding it to our mesh anyway.  This is important for degenerate cases.
+			++EarTestCount;
+		}
+	}
+
+	check( OutTriangles.Num() > 0 );
+}
+
+
+bool UEditableMesh::ComputeBarycentricWeightForPointOnPolygon( const FPolygonID PolygonID, const FVector PointOnPolygon, FMeshTriangle& OutTriangle, FVector& OutTriangleVertexWeights ) const
+{
+	// @todo mesheditor: Modify this method so it can return a meaningful barycentric weight for an off-polygon position, referencing the closest triangle, with out-of-range weights.
+
+	const FMeshPolygon& Polygon = Polygons[ PolygonID.GetValue() ];
+	const int32 TriangleCount = Polygon.Triangles.Num();
+
+	// Figure out which triangle the incoming point is within
+	for( const FMeshTriangle& Triangle : Polygon.Triangles )
+	{
+		const FVector TriangleVertex0Position = GetVertexInstanceAttribute( Triangle.VertexInstanceID0, UEditableMeshAttribute::VertexPosition(), 0 );
+		const FVector TriangleVertex1Position = GetVertexInstanceAttribute( Triangle.VertexInstanceID1, UEditableMeshAttribute::VertexPosition(), 0 );
+		const FVector TriangleVertex2Position = GetVertexInstanceAttribute( Triangle.VertexInstanceID2, UEditableMeshAttribute::VertexPosition(), 0 );
+
+		// Calculate the barycentric weights for the triangle's verts and determine if the point lies within its bounds.
+		OutTriangleVertexWeights = FMath::ComputeBaryCentric2D( PointOnPolygon, TriangleVertex0Position, TriangleVertex1Position, TriangleVertex2Position );
+
+		if( OutTriangleVertexWeights.X >= 0.0f && OutTriangleVertexWeights.Y >= 0.0f && OutTriangleVertexWeights.Z >= 0.0f )
+		{
+			// Okay, we found a triangle that the point is inside!  Return the corresponding vertex instances.
+			OutTriangle = Triangle;
+			return true;
+		}
+	}
+
+	return false;
+}
+
+
+void UEditableMesh::ComputeTextureCoordinatesForPointOnPolygon( const FPolygonID PolygonID, const FVector PointOnPolygon, bool& bOutWereTextureCoordinatesFound, TArray<FVector4>& OutInterpolatedTextureCoordinates )
+{
+	bOutWereTextureCoordinatesFound = false;
+	OutInterpolatedTextureCoordinates.Reset();
+
+	FMeshTriangle Triangle;
+	FVector TriangleVertexWeights;
+	if( ComputeBarycentricWeightForPointOnPolygon( PolygonID, PointOnPolygon, Triangle, TriangleVertexWeights ) )
+	{
+		OutInterpolatedTextureCoordinates.SetNum( TextureCoordinateCount, false );
+
+		for( int32 TextureCoordinateIndex = 0; TextureCoordinateIndex < TextureCoordinateCount; ++TextureCoordinateIndex )
+		{
+			OutInterpolatedTextureCoordinates[ TextureCoordinateIndex ] =
+				TriangleVertexWeights.X * GetVertexInstanceAttribute( Triangle.VertexInstanceID0, UEditableMeshAttribute::VertexTextureCoordinate(), TextureCoordinateIndex ) +
+				TriangleVertexWeights.Y * GetVertexInstanceAttribute( Triangle.VertexInstanceID1, UEditableMeshAttribute::VertexTextureCoordinate(), TextureCoordinateIndex ) +
+				TriangleVertexWeights.Z * GetVertexInstanceAttribute( Triangle.VertexInstanceID2, UEditableMeshAttribute::VertexTextureCoordinate(), TextureCoordinateIndex );
+		}
+
+		bOutWereTextureCoordinatesFound = true;
+	}
+}
+
+
+void UEditableMesh::SetSubdivisionCount( const int32 NewSubdivisionCount )
+{
+	// @todo mesheditor subdiv: Really, instead of a custom FChange type for this type of change, we could create a change that
+	// represents a property (or set of properties) being changed, and use Unreal reflection to update it.  This would be
+	// reusable for future properties.
+
+	const bool bEnablingSubdivisionPreview = ( GetSubdivisionCount() == 0 && NewSubdivisionCount > 0 );
+	const bool bDisablingSubdivisionPreview = ( GetSubdivisionCount() > 0 && NewSubdivisionCount == 0 );
+
+	FSetSubdivisionCountChangeInput RevertInput;
+	RevertInput.NewSubdivisionCount = GetSubdivisionCount();
+
+	this->SubdivisionCount = NewSubdivisionCount;
+
+	if( bDisablingSubdivisionPreview )
+	{
+		// We've turned off subdivision preview, so we'll need to re-create the static mesh data from our stored mesh representation
+		RebuildRenderMesh();
+	}
+	else
+	{
+		// NOTE: We don't bother regenerating geometry here because it's expected that EndModification() will be called after this, which will do the trick
+	}
+
+	AddUndo( MakeUnique<FSetSubdivisionCountChange>( MoveTemp( RevertInput ) ) );
+}
+
+
+void UEditableMesh::MoveVertices( const TArray<FVertexToMove>& VerticesToMove )
+{
+	static TSet< FPolygonID > VertexConnectedPolygons;
+	VertexConnectedPolygons.Reset();
+
+	static TArray<FAttributesForVertex> VertexAttributesToSet;
+	VertexAttributesToSet.Reset();
+
+	for( const FVertexToMove& VertexToMove : VerticesToMove )
+	{
+		const FVector CurrentPosition = GetVertexAttribute( VertexToMove.VertexID, UEditableMeshAttribute::VertexPosition(), 0 );
+
+		if( VertexToMove.NewVertexPosition != CurrentPosition )
+		{
+			FAttributesForVertex& AttributesForVertex = *new( VertexAttributesToSet ) FAttributesForVertex();
+			AttributesForVertex.VertexID = VertexToMove.VertexID;
+			FMeshElementAttributeData& VertexAttribute = *new( AttributesForVertex.VertexAttributes.Attributes ) FMeshElementAttributeData( UEditableMeshAttribute::VertexPosition(), 0, VertexToMove.NewVertexPosition );
+
+			// All of the polygons that share this vertex will need new normals
+			static TArray<FPolygonID> ConnectedPolygonRefs;
+			GetVertexConnectedPolygons( VertexToMove.VertexID, /* Out */ ConnectedPolygonRefs );
+			VertexConnectedPolygons.Append( ConnectedPolygonRefs );
+		}
+	}
+
+	SetVerticesAttributes( VertexAttributesToSet );
+
+	// Mark all polygons connected to the vertex as requiring a new tangent basis and retriangulation
+	// Everything needs to be retriangulated because convexity may have changed
+	PolygonsPendingNewTangentBasis.Append( VertexConnectedPolygons );
+	PolygonsPendingTriangulation.Append( VertexConnectedPolygons );
+}
+
+
+void UEditableMesh::CreateMissingPolygonPerimeterEdges( const FPolygonID PolygonID, TArray<FEdgeID>& OutNewEdgeIDs )
+{
+	// @todo mesheditor: Orphaned method - should we keep it or is it entirely unnecessary now?
+
+	OutNewEdgeIDs.Reset();
+
+	const int32 NumPolygonPerimeterEdges = GetPolygonPerimeterEdgeCount( PolygonID );
+	const int32 NumPolygonPerimeterVertices = NumPolygonPerimeterEdges;		// Edge and vertex count are always the same
+
+	for( int32 PerimeterEdgeNumber = 0; PerimeterEdgeNumber < NumPolygonPerimeterEdges; ++PerimeterEdgeNumber )
+	{
+		const int32 PerimeterVertexNumber = PerimeterEdgeNumber;	// Edge and vertex counts are always the same
+
+		const FVertexID VertexID = GetPolygonPerimeterVertex( PolygonID, PerimeterVertexNumber );
+		const FVertexID NextVertexID = GetPolygonPerimeterVertex( PolygonID, ( PerimeterVertexNumber + 1 ) % NumPolygonPerimeterVertices );
+
+		// Find the edge that connects these vertices
+		FEdgeID FoundEdgeID = FEdgeID::Invalid;
+		bool bFoundEdge = false;
+
+		const int32 NumVertexConnectedEdges = GetVertexConnectedEdgeCount( VertexID );	// @todo mesheditor perf: Could be made faster by always iterating over the vertex with the fewest number of connected edges, instead of the first edge vertex
+		for( int32 VertexEdgeNumber = 0; VertexEdgeNumber < NumVertexConnectedEdges; ++VertexEdgeNumber )
+		{
+			const FEdgeID VertexConnectedEdgeID = GetVertexConnectedEdge( VertexID, VertexEdgeNumber );
+
+			// Try the edge's first vertex.  Does it point to our next edge?
+			FVertexID OtherEdgeVertexID = GetEdgeVertex( VertexConnectedEdgeID, 0 );
+			if( OtherEdgeVertexID == VertexID )
+			{
+				// Must be the the other one
+				OtherEdgeVertexID = GetEdgeVertex( VertexConnectedEdgeID, 1 );
+			}
+			else
+			{
+				check( GetEdgeVertex( VertexConnectedEdgeID, 1 ) == VertexID );
+			}
+
+			if( OtherEdgeVertexID == NextVertexID )
+			{
+				// We found the edge!
+				FoundEdgeID = VertexConnectedEdgeID;
+				bFoundEdge = true;
+				break;
+			}
+		}
+
+		if( !bFoundEdge )
+		{
+			// Create the new edge!  Note that this does not connect the edge to the polygon.  We expect the caller to do that afterwards.
+			FEdgeToCreate EdgeToCreate;
+			EdgeToCreate.VertexID0 = VertexID;
+			EdgeToCreate.VertexID1 = NextVertexID;
+
+			static TArray< FEdgeToCreate > EdgesToCreate;
+			EdgesToCreate.Reset();
+			EdgesToCreate.Add( EdgeToCreate );
+
+			static TArray< FEdgeID > NewEdgeIDs;
+			NewEdgeIDs.Reset();
+			CreateEdges( EdgesToCreate, /* Out */ NewEdgeIDs );
+
+			OutNewEdgeIDs.Append( NewEdgeIDs );
+		}
+	}
+}
+
+
+void UEditableMesh::CreateMissingPolygonHoleEdges( const FPolygonID PolygonID, const int32 HoleNumber, TArray<FEdgeID>& OutNewEdgeIDs )
+{
+	// @todo mesheditor: Orphaned method - should we keep it or is it entirely unnecessary now?
+
+	OutNewEdgeIDs.Reset();
+
+	const int32 NumHoleEdges = GetPolygonHoleEdgeCount( PolygonID, HoleNumber );
+	const int32 NumHoleVertices = NumHoleEdges;		// Edge and vertex count are always the same
+
+	for( int32 HoleEdgeNumber = 0; HoleEdgeNumber < NumHoleEdges; ++HoleEdgeNumber )
+	{
+		const int32 HoleVertexNumber = HoleEdgeNumber;	// Edge and vertex counts are always the same
+
+		const FVertexID VertexID = GetPolygonHoleVertex( PolygonID, HoleNumber, HoleVertexNumber );
+		const FVertexID NextVertexID = GetPolygonHoleVertex( PolygonID, HoleNumber, ( HoleVertexNumber + 1 ) % NumHoleVertices );
+
+		// Find the edge that connects these vertices
+		FEdgeID FoundEdgeID = FEdgeID::Invalid;
+		{
+			bool bFoundEdge = false;
+
+			const int32 NumVertexConnectedEdges = GetVertexConnectedEdgeCount( VertexID );	// @todo mesheditor perf: Could be made faster by always iterating over the vertex with the fewest number of connected edges, instead of the first edge vertex
+			for( int32 VertexEdgeNumber = 0; VertexEdgeNumber < NumVertexConnectedEdges; ++VertexEdgeNumber )
+			{
+				const FEdgeID VertexConnectedEdgeID = GetVertexConnectedEdge( VertexID, VertexEdgeNumber );
+
+				// Try the edge's first vertex.  Does it point to us?
+				FVertexID OtherEdgeVertexID = GetEdgeVertex( VertexConnectedEdgeID, 0 );
+				if( OtherEdgeVertexID == VertexID )
+				{
+					// Must be the the other one
+					OtherEdgeVertexID = GetEdgeVertex( VertexConnectedEdgeID, 1 );
+				}
+
+				if( OtherEdgeVertexID == NextVertexID )
+				{
+					// We found the edge!
+					FoundEdgeID = VertexConnectedEdgeID;
+					bFoundEdge = true;
+					break;
+				}
+			}
+
+			if( !bFoundEdge )
+			{
+				// Create the new edge!  Note that this does not connect the edge to the polygon.  We expect the caller to do that afterwards.
+				FEdgeToCreate EdgeToCreate;
+				EdgeToCreate.VertexID0 = VertexID;
+				EdgeToCreate.VertexID1 = NextVertexID;
+
+				static TArray< FEdgeToCreate > EdgesToCreate;
+				EdgesToCreate.Reset();
+				EdgesToCreate.Add( EdgeToCreate );
+
+				static TArray< FEdgeID > NewEdgeIDs;
+				NewEdgeIDs.Reset();
+				CreateEdges( EdgesToCreate, /* Out */ NewEdgeIDs );
+
+				OutNewEdgeIDs.Append( NewEdgeIDs );
+			}
+		}
+	}
+}
+
+
+void UEditableMesh::SplitEdge( const FEdgeID EdgeID, const TArray<float>& Splits, TArray<FVertexID>& OutNewVertexIDs )
+{
+	UE_LOG( LogEditableMesh, Verbose, TEXT( "SplitEdge: %s" ), *EdgeID.ToString() );
+
+	// NOTE: The incoming splits should always be between 0.0 and 1.0, representing progress along 
+	//       the edge from the edge's first vertex toward it's other vertex.  The order doesn't matter (we'll sort them.)
+
+	check( Splits.Num() > 0 );
+
+	// Sort the split values smallest to largest.  We'll be adding a strip of vertices for each split, and the
+	// indices for those new vertices need to be in order.
+	static TArray<float> SortedSplits;
+	SortedSplits.Reset();
+	SortedSplits = Splits;
+	if( SortedSplits.Num() > 1 )
+	{
+		SortedSplits.Sort();
+	}
+
+	FVertexID OriginalEdgeVertexIDs[ 2 ];
+	GetEdgeVertices( EdgeID, /* Out */ OriginalEdgeVertexIDs[0], /* Out */ OriginalEdgeVertexIDs[1] );
+
+	static TArray<FVector> NewVertexPositions;
+	NewVertexPositions.Reset();
+	NewVertexPositions.AddUninitialized( SortedSplits.Num() );
+
+	for( int32 NewVertexNumber = 0; NewVertexNumber < NewVertexPositions.Num(); ++NewVertexNumber )
+	{
+		const float Split = SortedSplits[ NewVertexNumber ];
+		check( Split >= 0.0f && Split <= 1.0f );
+
+		NewVertexPositions[ NewVertexNumber ] = FMath::Lerp(
+			FVector( GetVertexAttribute( OriginalEdgeVertexIDs[ 0 ], UEditableMeshAttribute::VertexPosition(), 0 ) ),
+			FVector( GetVertexAttribute( OriginalEdgeVertexIDs[ 1 ], UEditableMeshAttribute::VertexPosition(), 0 ) ),
+			Split );
+	}
+
+	// Split the edge, and connect the vertex to the polygons that share the two new edges
+	const FVertexID OriginalEdgeFarVertexID = OriginalEdgeVertexIDs[ 1 ];
+
+	static TArray<FMeshElementAttributeData> OriginalEdgeAttributes;
+	OriginalEdgeAttributes.Reset();
+	{
+		for( const FName AttributeName : UEditableMesh::GetValidEdgeAttributes() )
+		{
+			const int32 MaxAttributeIndex = GetMaxAttributeIndex( AttributeName );
+			for( int32 AttributeIndex = 0; AttributeIndex < MaxAttributeIndex; ++AttributeIndex )
+			{
+				FMeshElementAttributeData& EdgeAttribute = *new( OriginalEdgeAttributes ) FMeshElementAttributeData(
+					AttributeName,
+					AttributeIndex,
+					GetEdgeAttribute( EdgeID, AttributeName, AttributeIndex ) );
+			}
+		}
+	}
+
+
+	// Create new vertices.  Their texture coordinates will be undefined, for now.  We'll set that later by using the connected 
+	// polygon's vertex data.  Also, the vertex's normals and tangents will be automatically generated.
+	static TArray<FVertexID> NewVertexIDs;
+	NewVertexIDs.Reset();
+	{
+		OutNewVertexIDs.Reset();
+		OutNewVertexIDs.Reserve( NewVertexPositions.Num() );
+
+		static TArray<FVertexToCreate> VerticesToCreate;
+		VerticesToCreate.Reset();
+		VerticesToCreate.Reserve( NewVertexPositions.Num() );
+
+		for( int32 NewVertexNumber = 0; NewVertexNumber < NewVertexPositions.Num(); ++NewVertexNumber )
+		{
+			const FVector& NewVertexPosition = NewVertexPositions[ NewVertexNumber ];
+
+			FVertexToCreate& VertexToCreate = *new( VerticesToCreate )FVertexToCreate();
+			VertexToCreate.VertexAttributes.Attributes.Add( FMeshElementAttributeData( UEditableMeshAttribute::VertexPosition(), 0, NewVertexPosition ) );
+		}
+
+		CreateVertices( VerticesToCreate, /* Out */ NewVertexIDs );
+
+		OutNewVertexIDs.Append( NewVertexIDs );
+	}
+
+	// Figure out which polygons are connected to the original edge.  Those polygons are going to need their
+	// connected vertices updated later on.
+	struct FAffectedPolygonEdge
+	{
+		FPolygonID PolygonID;
+		int32 PolygonVertexNumbers[ 2 ];
+	};
+	static TArray<FAffectedPolygonEdge> AffectedPolygonEdges;
+	AffectedPolygonEdges.Reset();
+
+	static TArray<FPolygonID> AffectedPolygons;
+	AffectedPolygons.Reset();
+
+	{
+		const int32 ConnectedPolygonCount = GetEdgeConnectedPolygonCount( EdgeID );
+		for( int32 ConnectedPolygonNumber = 0; ConnectedPolygonNumber < ConnectedPolygonCount; ++ConnectedPolygonNumber )
+		{
+			const FPolygonID PolygonID = GetEdgeConnectedPolygon( EdgeID, ConnectedPolygonNumber );
+
+			// Figure out which of this polygon's vertices are part of the edge we're adding a vertex to
+			// @todo mesheditor hole: Add support for polygon holes, too!
+			static TArray<FVertexID> PolygonPerimeterVertexIDs;
+			PolygonPerimeterVertexIDs.Reset();
+			GetPolygonPerimeterVertices( PolygonID, /* Out */ PolygonPerimeterVertexIDs );
+
+			int32 PolygonVertexNumberForEdgeVertices[ 2 ] = { INDEX_NONE, INDEX_NONE };
+			for( int32 PerimeterVertexNumber = 0; PerimeterVertexNumber < PolygonPerimeterVertexIDs.Num(); ++PerimeterVertexNumber )
+			{
+				const FVertexID PerimeterVertexID = PolygonPerimeterVertexIDs[ PerimeterVertexNumber ];
+				if( PerimeterVertexID == OriginalEdgeVertexIDs[ 0 ] )
+				{
+					// Okay, this polygon vertex is part of our edge (first edge vertex)
+					PolygonVertexNumberForEdgeVertices[ 0 ] = PerimeterVertexNumber;
+				}
+				else if( PerimeterVertexID == OriginalEdgeVertexIDs[ 1 ] )
+				{
+					// Okay, this polygon vertex is part of our edge (second edge vertex)
+					PolygonVertexNumberForEdgeVertices[ 1 ] = PerimeterVertexNumber;
+				}
+			}
+			check( PolygonVertexNumberForEdgeVertices[ 0 ] != INDEX_NONE && PolygonVertexNumberForEdgeVertices[ 1 ] != INDEX_NONE );
+
+			FAffectedPolygonEdge& AffectedPolygonEdge = *new( AffectedPolygonEdges ) FAffectedPolygonEdge();
+			AffectedPolygonEdge.PolygonID = PolygonID;
+			AffectedPolygonEdge.PolygonVertexNumbers[ 0 ] = PolygonVertexNumberForEdgeVertices[ 0 ];
+			AffectedPolygonEdge.PolygonVertexNumbers[ 1 ] = PolygonVertexNumberForEdgeVertices[ 1 ];
+
+			// @todo mesheditor perf: Not sure if we actually need to AddUnique() here -- can an edge be connected to the same polygon more than once? (tube shape). Seems sketchy if so?  And if so, GetEdgeConnectedPolygonCount() probably should have only reported 1 connection anyway
+			AffectedPolygons.AddUnique( AffectedPolygonEdge.PolygonID );
+		}
+	}
+
+	{
+		static TArray<FVerticesForEdge> VerticesForEdges;
+		VerticesForEdges.Reset();
+
+		// We'll keep the existing edge, but update it to connect to the first new vertex.  The second vertex of the edge will 
+		// now connect to the first (new) vertex ID, and so on.  The incoming vertices are expected to be ordered correctly.
+		{
+			FVerticesForEdge& VerticesForEdge = *new( VerticesForEdges ) FVerticesForEdge();
+
+			VerticesForEdge.EdgeID = EdgeID;
+			VerticesForEdge.NewVertexID0 = OriginalEdgeVertexIDs[ 0 ];
+			VerticesForEdge.NewVertexID1 = NewVertexIDs[ 0 ];
+		}
+
+		SetEdgesVertices( VerticesForEdges );
+	}
+
+	// Create new edges.  One for each of the new vertex positions passed in.
+	{
+		const int32 NewEdgeCount = NewVertexPositions.Num();
+
+		static TArray<FEdgeToCreate> EdgesToCreate;
+		EdgesToCreate.Reset();
+		EdgesToCreate.Reserve( NewEdgeCount );
+		for( int32 NewEdgeNumber = 0; NewEdgeNumber < NewEdgeCount; ++NewEdgeNumber )
+		{
+			FEdgeToCreate& EdgeToCreate = *new( EdgesToCreate ) FEdgeToCreate();
+
+			EdgeToCreate.VertexID0 = NewVertexIDs[ NewEdgeNumber ];
+			EdgeToCreate.VertexID1 = ( NewEdgeNumber == ( NewEdgeCount - 1 ) ) ? OriginalEdgeFarVertexID : NewVertexIDs[ NewEdgeNumber + 1 ];
+
+			EdgeToCreate.ConnectedPolygons = AffectedPolygons;
+
+			// Copy edge attributes over from original edge
+			EdgeToCreate.EdgeAttributes.Attributes = OriginalEdgeAttributes;
+		}
+
+		static TArray< FEdgeID > NewEdgeIDs;
+		NewEdgeIDs.Reset();
+		CreateEdges( EdgesToCreate, /* Out */ NewEdgeIDs );
+	}
+
+	// Update all affected polygons with their new vertices.  Also, we'll fill in polygon-specific vertex attributes (texture coordinates)
+	{
+		for( const FAffectedPolygonEdge& AffectedPolygonEdge : AffectedPolygonEdges )
+		{
+			const FPolygonID PolygonID = AffectedPolygonEdge.PolygonID;
+
+			// Figure out where to start inserting vertices into this polygon, as well as which order to insert them.
+			bool bWindsForward = AffectedPolygonEdge.PolygonVertexNumbers[ 1 ] > AffectedPolygonEdge.PolygonVertexNumbers[ 0 ];
+			const int32 LargerIndex = bWindsForward ? AffectedPolygonEdge.PolygonVertexNumbers[ 1 ] : AffectedPolygonEdge.PolygonVertexNumbers[ 0 ];
+			const bool bIsContiguous = FMath::Abs( (int32)AffectedPolygonEdge.PolygonVertexNumbers[ 1 ] - (int32)AffectedPolygonEdge.PolygonVertexNumbers[ 0 ] ) == 1;
+			if( !bIsContiguous )
+			{
+				bWindsForward = !bWindsForward;
+			}
+
+			const int32 InsertAtIndex = bIsContiguous ? LargerIndex : ( LargerIndex + 1 );
+
+			// Figure out what the polygon's NEW vertex count (after we insert new vertices) will be
+			const int32 NewPolygonPerimeterVertexCount = GetPolygonPerimeterVertexCount( PolygonID ) + NewVertexPositions.Num();
+
+			static TArray<FVertexAndAttributes> VerticesToInsert;
+			VerticesToInsert.Reset( NewVertexPositions.Num() );
+			VerticesToInsert.SetNum( NewVertexPositions.Num(), false );
+			for( int32 InsertVertexNumber = 0; InsertVertexNumber < NewVertexPositions.Num(); ++InsertVertexNumber )
+			{
+				FVertexAndAttributes& VertexToInsert = VerticesToInsert[ InsertVertexNumber ];
+
+				const int32 DirectionalVertexNumber = bWindsForward ? InsertVertexNumber : ( ( NewVertexPositions.Num() - 1 ) - InsertVertexNumber );
+				VertexToInsert.VertexID = NewVertexIDs[ InsertVertexNumber ];
+				VertexToInsert.VertexInstanceID = FVertexInstanceID::Invalid;
+
+				// Generate texture coordinates for this vertex by interpolating between the polygon vertex texture coordinates on this polygon edge
+				{
+					const float Split = SortedSplits[ DirectionalVertexNumber ];
+					for( int32 TextureCoordinateIndex = 0; TextureCoordinateIndex < TextureCoordinateCount; ++TextureCoordinateIndex )
+					{
+						const FVector4 NewPolygonVertexTextureCoordinate = FMath::Lerp(
+							GetPolygonPerimeterVertexAttribute( PolygonID, AffectedPolygonEdge.PolygonVertexNumbers[ 0 ], UEditableMeshAttribute::VertexTextureCoordinate(), TextureCoordinateIndex ),
+							GetPolygonPerimeterVertexAttribute( PolygonID, AffectedPolygonEdge.PolygonVertexNumbers[ 1 ], UEditableMeshAttribute::VertexTextureCoordinate(), TextureCoordinateIndex ),
+							Split );
+
+						VertexToInsert.PolygonVertexAttributes.Attributes.Add( FMeshElementAttributeData( UEditableMeshAttribute::VertexTextureCoordinate(), TextureCoordinateIndex, NewPolygonVertexTextureCoordinate ) );
+					}
+				}
+
+				// Generate vertex color for this vertex by interpolating between the vertex colors on the edge
+				{
+					const float Split = SortedSplits[ DirectionalVertexNumber ];
+					const FVector4 NewVertexColor = FMath::Lerp(
+						GetPolygonPerimeterVertexAttribute( PolygonID, AffectedPolygonEdge.PolygonVertexNumbers[ 0 ], UEditableMeshAttribute::VertexColor(), 0 ),
+						GetPolygonPerimeterVertexAttribute( PolygonID, AffectedPolygonEdge.PolygonVertexNumbers[ 1 ], UEditableMeshAttribute::VertexColor(), 0 ),
+						Split );
+
+					VertexToInsert.PolygonVertexAttributes.Attributes.Add( FMeshElementAttributeData( UEditableMeshAttribute::VertexColor(), 0, NewVertexColor ) );
+				}
+			}
+
+			// Add the new vertices to the polygon, and set texture coordinates for all of the affected polygon vertices
+			InsertPolygonPerimeterVertices( PolygonID, InsertAtIndex, VerticesToInsert );
+		}
+	}
+
+	// Generate normals and tangents
+	PolygonsPendingNewTangentBasis.Append( AffectedPolygons );
+
+	// Retriangulate all of the affected polygons
+	PolygonsPendingTriangulation.Append( AffectedPolygons );
+
+	UE_LOG( LogEditableMesh, Verbose, TEXT( "* SplitEdge returned %s" ), *LogHelpers::ArrayToString( OutNewVertexIDs ) );
+}
+
+
+void UEditableMesh::FindPolygonLoop( const FEdgeID EdgeID, TArray<FEdgeID>& OutEdgeLoopEdgeIDs, TArray<FEdgeID>& OutFlippedEdgeIDs, TArray<FEdgeID>& OutReversedEdgeIDPathToTake, TArray<FPolygonID>& OutPolygonIDsToSplit ) const
+{
+	OutEdgeLoopEdgeIDs.Reset();
+	OutFlippedEdgeIDs.Reset();
+	OutReversedEdgeIDPathToTake.Reset();
+	OutPolygonIDsToSplit.Reset();
+
+	// Is the edge we're starting on a border edge?
+	bool bStartedOnBorderEdge = ( GetEdgeConnectedPolygonCount( EdgeID ) <= 1 );
+
+	// We'll actually do two passes searching for edges.  The first time, we'll flow along the polygons looking for a border edge.
+	// If we find one, it means we won't really have a loop, but instead we're just splitting a series of connected polygons.  In that
+	// case, the polygon with the border edge will become the start of our search for an opposing border edge.  Otherwise, we'll start
+	// at the input polygon and flow around until we come back to that polygon again.  If we can't make it back, then no polygons
+	// will be split by this operation.
+	bool bIsSearchingForBorderEdge = !bStartedOnBorderEdge;
+
+	// Keep track of whether we actually looped back around to the starting edge (rather than simply splitting across a series
+	// of polygons that both end at border edges.)
+	bool bIsCompleteLoop = false;
+
+	FEdgeID CurrentEdgeID = EdgeID;
+	bool bCurrentEdgeIsBorderEdge = bStartedOnBorderEdge;
+	bool bCurrentEdgeIsInOppositeDirection = false;
+	bool bCurrentEdgeIsInOppositeDirectionFromStartEdge = false;
+	for( ; ; )
+	{
+		// Add the current edge!
+		checkSlow( !OutEdgeLoopEdgeIDs.Contains( CurrentEdgeID ) );
+		OutEdgeLoopEdgeIDs.Add( CurrentEdgeID );
+		if( bCurrentEdgeIsInOppositeDirection )
+		{
+			OutFlippedEdgeIDs.Add( CurrentEdgeID );
+		}
+
+		FVertexID CurrentEdgeVertexIDs[ 2 ];
+		GetEdgeVertices( CurrentEdgeID, /* Out */ CurrentEdgeVertexIDs[ 0 ], /* Out */ CurrentEdgeVertexIDs[ 1 ] );
+
+		const FVector CurrentEdgeVertex0 = GetVertexAttribute( CurrentEdgeVertexIDs[ 0 ], UEditableMeshAttribute::VertexPosition(), 0 );
+		const FVector CurrentEdgeVertex1 = GetVertexAttribute( CurrentEdgeVertexIDs[ 1 ], UEditableMeshAttribute::VertexPosition(), 0 );
+		const FVector CurrentEdgeDirection = ( CurrentEdgeVertex1 - CurrentEdgeVertex0 ).GetSafeNormal();
+
+		const FEdgeID NextEdgeIDInPath = OutReversedEdgeIDPathToTake.Num() > 0 ? OutReversedEdgeIDPathToTake.Pop( false ) : FEdgeID::Invalid;
+
+		FEdgeID BestEdgeID = FEdgeID::Invalid;
+		FPolygonID BestEdgeSplitsPolygon = FPolygonID::Invalid;
+		bool bBestEdgeIsInOppositeDirection = false;
+		bool bBestEdgeIsBorderEdge = false;
+		float LargestAbsDotProduct = -1.0f;
+
+		// Let's take a look at all of the polygons connected to this edge.  These will start our loop.
+		const int32 ConnectedPolygonCount = GetEdgeConnectedPolygonCount( CurrentEdgeID );
+		for( int32 ConnectedPolygonNumber = 0; ConnectedPolygonNumber < ConnectedPolygonCount; ++ConnectedPolygonNumber )
+		{
+			const FPolygonID ConnectedPolygonID = GetEdgeConnectedPolygon( CurrentEdgeID, ConnectedPolygonNumber );
+
+			// Don't bother looking at the last polygon that was added to our split list.  We never want to backtrack!
+			if( OutPolygonIDsToSplit.Num() == 0 || ConnectedPolygonID != OutPolygonIDsToSplit.Last() )
+			{
+				static TArray<FEdgeID> CandidateEdgeIDs;
+				CandidateEdgeIDs.Reset();
+				GetPolygonPerimeterEdges( ConnectedPolygonID, /* Out */ CandidateEdgeIDs );
+
+				// Which edge of the connected polygon will be at the other end of our split?
+				for( const FEdgeID CandidateEdgeID : CandidateEdgeIDs )
+				{
+					// Don't bother with the edge we just came from
+					if( CandidateEdgeID != CurrentEdgeID )
+					{
+						// If we need to follow a specific path, then do that
+						if( NextEdgeIDInPath == FEdgeID::Invalid || CandidateEdgeID == NextEdgeIDInPath )
+						{
+							FVertexID CandidateEdgeVertexIDs[ 2 ];
+							GetEdgeVertices( CandidateEdgeID, /* Out */ CandidateEdgeVertexIDs[ 0 ], /* Out */ CandidateEdgeVertexIDs[ 1 ] );
+
+							const bool bIsBorderEdge = ( GetEdgeConnectedPolygonCount( CandidateEdgeID ) == 1 );
+
+							const FVector CandidateEdgeVertex0 = GetVertexAttribute( CandidateEdgeVertexIDs[ 0 ], UEditableMeshAttribute::VertexPosition(), 0 );
+							const FVector CandidateEdgeVertex1 = GetVertexAttribute( CandidateEdgeVertexIDs[ 1 ], UEditableMeshAttribute::VertexPosition(), 0 );
+							const FVector CandidateEdgeDirection = ( CandidateEdgeVertex1 - CandidateEdgeVertex0 ).GetSafeNormal();
+
+							const float DotProduct = FVector::DotProduct( CurrentEdgeDirection, CandidateEdgeDirection );
+							const float AbsDotProduct = FMath::Abs( DotProduct );
+
+							const float SameEdgeDirectionDotEpsilon = 0.05f;	// @todo mesheditor tweak
+							if( FMath::IsNearlyEqual( AbsDotProduct, LargestAbsDotProduct, SameEdgeDirectionDotEpsilon ) )
+							{
+								// If the candidate edge directions are pretty much the same, we'll choose the edge that flows closest to the
+								// direction that we split the last polygon in
+								if( OutEdgeLoopEdgeIDs.Num() > 1 )
+								{
+									FVertexID LastSplitEdgeVertexIDs[ 2 ];
+									GetEdgeVertices( OutEdgeLoopEdgeIDs[ OutEdgeLoopEdgeIDs.Num() - 2 ], /* Out */ LastSplitEdgeVertexIDs[ 0 ], /* Out */ LastSplitEdgeVertexIDs[ 1 ] );
+
+									const FVector LastSplitEdgeVertex0 = GetVertexAttribute( LastSplitEdgeVertexIDs[ 0 ], UEditableMeshAttribute::VertexPosition(), 0 );
+									const FVector LastSplitEdgeVertex1 = GetVertexAttribute( LastSplitEdgeVertexIDs[ 1 ], UEditableMeshAttribute::VertexPosition(), 0 );
+
+									const FVector DirectionTowardCenterOfCurrentEdge =
+										( FMath::Lerp( CurrentEdgeVertex0, CurrentEdgeVertex1, 0.5f ) -
+											FMath::Lerp( LastSplitEdgeVertex0, LastSplitEdgeVertex1, 0.5f ) ).GetSafeNormal();
+
+									const FVector DirectionTowardCenterOfCandidateEdge =
+										( FMath::Lerp( CandidateEdgeVertex0, CandidateEdgeVertex1, 0.5f ) -
+											FMath::Lerp( CurrentEdgeVertex0, CurrentEdgeVertex1, 0.5f ) ).GetSafeNormal();
+									const float CandidateEdgeDot = FVector::DotProduct( DirectionTowardCenterOfCurrentEdge, DirectionTowardCenterOfCandidateEdge );
+
+									check( BestEdgeID != FEdgeID::Invalid );
+
+									FVertexID BestEdgeVertexIDs[ 2 ];
+									GetEdgeVertices( BestEdgeID, /* Out */ BestEdgeVertexIDs[ 0 ], /* Out */ BestEdgeVertexIDs[ 1 ] );
+
+									const FVector BestEdgeVertex0 = GetVertexAttribute( BestEdgeVertexIDs[ 0 ], UEditableMeshAttribute::VertexPosition(), 0 );
+									const FVector BestEdgeVertex1 = GetVertexAttribute( BestEdgeVertexIDs[ 1 ], UEditableMeshAttribute::VertexPosition(), 0 );
+
+									const FVector DirectionTowardCenterOfBestEdge =
+										( FMath::Lerp( BestEdgeVertex0, BestEdgeVertex1, 0.5f ) -
+											FMath::Lerp( CurrentEdgeVertex0, CurrentEdgeVertex1, 0.5f ) ).GetSafeNormal();
+
+									const float BestEdgeDot = FVector::DotProduct( DirectionTowardCenterOfCurrentEdge, DirectionTowardCenterOfBestEdge );
+
+									if( CandidateEdgeDot > BestEdgeDot )
+									{
+										BestEdgeID = CandidateEdgeID;
+										BestEdgeSplitsPolygon = ConnectedPolygonID;
+										bBestEdgeIsInOppositeDirection = ( DotProduct < 0.0f );
+										bBestEdgeIsBorderEdge = bIsBorderEdge;
+										LargestAbsDotProduct = AbsDotProduct;
+									}
+								}
+								else
+								{
+									// Edge directions are the same, but this is the very first split so we don't have a "flow" direction yet.
+									// Go ahead and prefer the edge that is closer to the initial edge.  This helps in the (uncommon) case of 
+									// multiple colinear edges on the same polygon (such as after SplitEdge() is called to insert a vertex on a polygon.)
+									const float BestEdgeDistance = [this, BestEdgeID, CurrentEdgeVertex0, CurrentEdgeVertex1]() -> float
+									{
+										check( BestEdgeID != FEdgeID::Invalid );
+
+										FVertexID BestEdgeVertexIDs[ 2 ];
+										GetEdgeVertices( BestEdgeID, /* Out */ BestEdgeVertexIDs[ 0 ], /* Out */ BestEdgeVertexIDs[ 1 ] );
+
+										const FVector BestEdgeVertex0 = GetVertexAttribute( BestEdgeVertexIDs[ 0 ], UEditableMeshAttribute::VertexPosition(), 0 );
+										const FVector BestEdgeVertex1 = GetVertexAttribute( BestEdgeVertexIDs[ 1 ], UEditableMeshAttribute::VertexPosition(), 0 );
+
+										FVector ClosestPoint0, ClosestPoint1;
+										FMath::SegmentDistToSegmentSafe( CurrentEdgeVertex0, CurrentEdgeVertex1, BestEdgeVertex0, BestEdgeVertex1, /* Out */ ClosestPoint0, /* Out */ ClosestPoint1 );
+										return ( ClosestPoint1 - ClosestPoint0 ).Size();
+									}( );
+
+									const float CandidateEdgeDistance = [CurrentEdgeVertex0, CurrentEdgeVertex1, CandidateEdgeVertex0, CandidateEdgeVertex1]() -> float
+									{
+										FVector ClosestPoint0, ClosestPoint1;
+										FMath::SegmentDistToSegmentSafe( CurrentEdgeVertex0, CurrentEdgeVertex1, CandidateEdgeVertex0, CandidateEdgeVertex1, /* Out */ ClosestPoint0, /* Out */ ClosestPoint1 );
+										return ( ClosestPoint1 - ClosestPoint0 ).Size();
+									}( );
+
+									if( CandidateEdgeDistance < BestEdgeDistance )
+									{
+										BestEdgeID = CandidateEdgeID;
+										BestEdgeSplitsPolygon = ConnectedPolygonID;
+										bBestEdgeIsInOppositeDirection = ( DotProduct < 0.0f );
+										bBestEdgeIsBorderEdge = bIsBorderEdge;
+										LargestAbsDotProduct = AbsDotProduct;
+									}
+								}
+							}
+							else if( AbsDotProduct > LargestAbsDotProduct )
+							{
+								// This edge angle is the closest to our current edge so far!
+								BestEdgeID = CandidateEdgeID;
+								BestEdgeSplitsPolygon = ConnectedPolygonID;
+								bBestEdgeIsInOppositeDirection = ( DotProduct < 0.0f );
+								bBestEdgeIsBorderEdge = bIsBorderEdge;
+								LargestAbsDotProduct = AbsDotProduct;
+							}
+						}
+					}
+				}
+			}
+		}
+
+		if( BestEdgeID != FEdgeID::Invalid &&
+			!OutPolygonIDsToSplit.Contains( BestEdgeSplitsPolygon ) )	// If we try to re-split the same polygon twice, then this loop is not valid
+		{
+			// OK, this polygon will definitely be split
+			OutPolygonIDsToSplit.Add( BestEdgeSplitsPolygon );
+
+			CurrentEdgeID = BestEdgeID;
+			bCurrentEdgeIsBorderEdge = bBestEdgeIsBorderEdge;
+			bCurrentEdgeIsInOppositeDirection = bBestEdgeIsInOppositeDirection;
+			if( bBestEdgeIsInOppositeDirection )
+			{
+				bCurrentEdgeIsInOppositeDirectionFromStartEdge = !bCurrentEdgeIsInOppositeDirectionFromStartEdge;
+			}
+
+			// Is the best edge already part of our loop?  If so, then we're done!
+			if( OutEdgeLoopEdgeIDs[ 0 ] == BestEdgeID )
+			{
+				bIsCompleteLoop = true;
+				break;
+			}
+			else if( OutEdgeLoopEdgeIDs.Contains( BestEdgeID ) )	// @todo mesheditor perf: Might want to use a TSet here, but we'll still need to keep a list too (order is important.)
+			{
+				// We ended up back at an edge that we already split, but it wasn't the edge that we started on.  This is
+				// not a valid loop, so clear our path and bail out.
+
+				// @todo mesheditor edgeloop: We need to revisit how we handle border edges and non-manifold edges.  Basically instead of searching
+				// for border edges, we should try looping around the "back" of polygons we've already visited, trying to get back to the starting
+				// edge.  This will allow the case of a cube mesh with a single non-manifold face to be split properly.  We just need to make sure
+				// we're not trying to split the same edge multiple times in different locations.  Probably, we don't want to split the same polygon
+				// more than once either.
+				OutEdgeLoopEdgeIDs.Reset();
+				OutFlippedEdgeIDs.Reset();
+				OutPolygonIDsToSplit.Reset();
+
+				break;
+			}
+			else if( bBestEdgeIsBorderEdge && bIsSearchingForBorderEdge )
+			{
+				// We found a border edge, so stop the search.  We'll now start over at this edge to form our loop.
+				bStartedOnBorderEdge = true;
+
+				bIsSearchingForBorderEdge = false;
+				bCurrentEdgeIsInOppositeDirection = bCurrentEdgeIsInOppositeDirectionFromStartEdge;
+
+				// Follow the path we took to get here, in reverse order, to make sure we get back to the edge we 
+				// were asked to create a loop on
+				OutReversedEdgeIDPathToTake = OutEdgeLoopEdgeIDs;
+
+				bIsCompleteLoop = false;
+
+				OutEdgeLoopEdgeIDs.Reset();
+				OutFlippedEdgeIDs.Reset();
+				OutPolygonIDsToSplit.Reset();
+			}
+			else
+			{
+				// Proceed to the next edge and try to continue the loop.  If we're at a border edge, the loop will definitely end here.
+			}
+		}
+		else
+		{
+			if( bStartedOnBorderEdge && bCurrentEdgeIsBorderEdge )
+			{
+				// We started on a border edge, and we've found the border edge on the other side of the polygons we'll be splitting.
+				// This isn't actually a loop, but we'll still split the polygons.
+			}
+			else
+			{
+				// We couldn't even find another edge, so we're done.
+				OutEdgeLoopEdgeIDs.Reset();
+				OutFlippedEdgeIDs.Reset();
+				OutPolygonIDsToSplit.Reset();
+			}
+
+			break;
+		}
+	}
+
+	// We're always splitting the same number of polygons as we have edges in the loop (these can be zero), except
+	// in the border edge case, where we're always splitting one less polygon.
+	if( bStartedOnBorderEdge &&
+		!bIsCompleteLoop )	// It's possible to start and finish on the same border edge, which is a complete loop
+	{
+		// We're splitting a series of polygons between two border edges
+		check( ( OutEdgeLoopEdgeIDs.Num() == 0 && OutPolygonIDsToSplit.Num() == 0 ) || ( OutEdgeLoopEdgeIDs.Num() == ( OutPolygonIDsToSplit.Num() + 1 ) ) );
+	}
+	else
+	{
+		// We're splitting polygons spanning a full loop of edges.  The starting edge is the same as the beginning edge.
+		check( OutEdgeLoopEdgeIDs.Num() == OutPolygonIDsToSplit.Num() );
+	}
+}
+
+
+void UEditableMesh::InsertEdgeLoop( const FEdgeID EdgeID, const TArray<float>& Splits, TArray<FEdgeID>& OutNewEdgeIDs )
+{
+	OutNewEdgeIDs.Reset();
+
+	// NOTE: The incoming splits should always be between 0.0 and 1.0, representing progress along 
+	//       the edge from the edge's first vertex toward it's other vertex.  The order doesn't matter (we'll sort them.)
+
+	// @todo mesheditor: Test with concave polygons -- probably we need to disallow/avoid splits that cross a concave part
+
+
+	static TArray<FEdgeID> EdgeLoopEdgeIDs;
+	EdgeLoopEdgeIDs.Reset();
+
+	static TArray<FEdgeID> FlippedEdgeIDs;
+	FlippedEdgeIDs.Reset();
+
+	static TArray<FEdgeID> ReversedEdgeIDPathToTake;
+	ReversedEdgeIDPathToTake.Reset();
+
+	static TArray<FPolygonID> PolygonIDsToSplit;
+	PolygonIDsToSplit.Reset();
+
+
+	FindPolygonLoop( 
+		EdgeID, 
+		/* Out */ EdgeLoopEdgeIDs,
+		/* Out */ FlippedEdgeIDs,
+		/* Out */ ReversedEdgeIDPathToTake,
+		/* Out */ PolygonIDsToSplit );
+
+	static TSet<FEdgeID> FlippedEdgeIDSet;
+	FlippedEdgeIDSet.Reset();
+	FlippedEdgeIDSet.Append( FlippedEdgeIDs );
+
+	check( Splits.Num() > 0 );
+
+	// Sort the split values smallest to largest.  We'll be adding a strip of vertices for each split, and the
+	// IDs for those new vertices need to be in order.
+	static TArray<float> SortedSplits;
+	SortedSplits.Reset();
+	SortedSplits = Splits;
+	if( SortedSplits.Num() > 1 )
+	{
+		SortedSplits.Sort();
+	}
+
+
+	if( PolygonIDsToSplit.Num() > 0 )
+	{
+		// Keep track of the new vertices create by splitting all of the edges.  For each edge we split, an array of
+		// vertex IDs for each split along that edge
+		static TArray<TArray<FVertexID>> NewVertexIDsForEachEdge;
+		NewVertexIDsForEachEdge.Reset();
+
+
+		// Now let's go through and create new vertices for the loops by splitting edges
+		{
+			for( int32 EdgeLoopEdgeNumber = 0; EdgeLoopEdgeNumber < EdgeLoopEdgeIDs.Num(); ++EdgeLoopEdgeNumber )
+			{
+				const FEdgeID EdgeLoopEdgeID = EdgeLoopEdgeIDs[ EdgeLoopEdgeNumber ];
+
+				// If the edge winds in the opposite direction from the last edge, we'll need to flip the split positions around
+				const bool bIsFlipped = FlippedEdgeIDSet.Contains( EdgeLoopEdgeID );
+				if( bIsFlipped )
+				{
+					static TArray<float> TempSplits;
+					TempSplits.SetNumUninitialized( SortedSplits.Num(), false );
+					for( int32 SplitIndex = 0; SplitIndex < SortedSplits.Num(); ++SplitIndex )
+					{
+						TempSplits[ SplitIndex ] = 1.0f - SortedSplits[ ( SortedSplits.Num() - 1 ) - SplitIndex ];
+					}
+					SortedSplits = TempSplits;
+				}
+
+				// Split this edge
+				static TArray<FVertexID> CurrentVertexIDs;
+				CurrentVertexIDs.Reset();
+				SplitEdge( EdgeLoopEdgeID, SortedSplits, /* Out */ CurrentVertexIDs );
+
+				// If the edge winding is backwards, we'll reverse the order of the vertex IDs in our list
+				// @todo mesheditor urgent edgeloop: INCORRECT (bIsFlipped is supposed to swap from last time -- TempSplits.)  Also not sure if needed or not... might be double compensating here.
+				static TArray<FVertexID> NewVertexIDsForEdge;
+				NewVertexIDsForEdge.SetNum( CurrentVertexIDs.Num(), false );
+				for( int32 VertexNumber = 0; VertexNumber < CurrentVertexIDs.Num(); ++VertexNumber )
+				{
+					NewVertexIDsForEdge[ /*bIsFlipped ? */( ( CurrentVertexIDs.Num() - VertexNumber ) - 1 ) /*: VertexNumber */] = CurrentVertexIDs[ VertexNumber ];
+				}
+				NewVertexIDsForEachEdge.Add( NewVertexIDsForEdge );
+			}
+		}
+
+
+		// Time to create new polygons for the split faces (and delete the old ones)
+		{
+			static TArray<FPolygonToSplit> PolygonsToSplit;
+			PolygonsToSplit.Reset();
+
+			for( int32 PolygonToSplitIter = 0; PolygonToSplitIter < PolygonIDsToSplit.Num(); ++PolygonToSplitIter )
+			{
+				const FPolygonID PolygonID = PolygonIDsToSplit[ PolygonToSplitIter ];
+
+				FPolygonToSplit& PolygonToSplit = *new( PolygonsToSplit ) FPolygonToSplit();
+				PolygonToSplit.PolygonID = PolygonID;
+
+				// The first and second edges connected to this polygon that are being split up
+				const int32 FirstEdgeNumber = PolygonToSplitIter;
+				const int32 SecondEdgeNumber = ( PolygonToSplitIter + 1 ) % EdgeLoopEdgeIDs.Num();
+
+				const FEdgeID FirstSplitEdgeID = EdgeLoopEdgeIDs[ FirstEdgeNumber ];
+				const FEdgeID SecondSplitEdgeID = EdgeLoopEdgeIDs[ SecondEdgeNumber ];
+				check( FirstSplitEdgeID != SecondSplitEdgeID );
+
+				// The (ordered) list of new vertices that was created by splitting the first and second edge.  One for each split.
+				const TArray<FVertexID>& FirstSplitEdgeNewVertexIDs = NewVertexIDsForEachEdge[ FirstEdgeNumber ];
+				const TArray<FVertexID>& SecondSplitEdgeNewVertexIDs = NewVertexIDsForEachEdge[ SecondEdgeNumber ];
+
+				for( int32 SplitIter = 0; SplitIter < SortedSplits.Num(); ++SplitIter )
+				{
+					FVertexPair& NewVertexPair = *new( PolygonToSplit.VertexPairsToSplitAt ) FVertexPair();
+					NewVertexPair.VertexID0 = FirstSplitEdgeNewVertexIDs[ SplitIter ];
+					NewVertexPair.VertexID1 = SecondSplitEdgeNewVertexIDs[ SplitIter ];
+				}
+			}
+
+
+			// Actually split up the polygons
+			static TArray<FEdgeID> NewEdgeIDs;
+			NewEdgeIDs.Reset();
+			SplitPolygons( PolygonsToSplit, /* Out */ NewEdgeIDs );
+
+			OutNewEdgeIDs.Append( NewEdgeIDs );
+		}
+	}
+}
+
+
+void UEditableMesh::SplitPolygons( const TArray<FPolygonToSplit>& PolygonsToSplit, TArray<FEdgeID>& OutNewEdgeIDs )
+{
+	OutNewEdgeIDs.Reset();
+
+	static TArray<FPolygonToCreate> PolygonsToCreate;
+	PolygonsToCreate.Reset();
+
+	for( const FPolygonToSplit& PolygonToSplit : PolygonsToSplit )
+	{
+		const FPolygonID PolygonID = PolygonToSplit.PolygonID;
+
+		// Get all of the polygon's vertices
+		static TArray<FVertexID> PerimeterVertexIDs;
+		PerimeterVertexIDs.Reset();
+		GetPolygonPerimeterVertices( PolygonID, /* Out */ PerimeterVertexIDs );
+		
+
+		// Figure out where exactly we're splitting the polygon for these splits.  Remember, we support splitting the
+		// polygon multiple times at once.  The first and last split are the most interesting because we need to continue
+		// the original flow of the polygon after inserting our new edge.  For all of the new polygons in the middle, we'll 
+		// just create simple quads.
+		const int32 SplitCount = PolygonToSplit.VertexPairsToSplitAt.Num();
+
+		int32 LastPolygonVertexNumbers[ 2 ] = { INDEX_NONE, INDEX_NONE };
+		bool bLastPolygonWindsForward = false;
+
+		const int32 NumPolygonsToCreate = SplitCount + 1;
+		for( int32 PolygonIter = 0; PolygonIter < NumPolygonsToCreate; ++PolygonIter )
+		{
+			const FVertexPair& VertexPair = PolygonToSplit.VertexPairsToSplitAt[ FMath::Min( PolygonIter, NumPolygonsToCreate - 2 ) ];
+
+			const FVertexID FirstVertexID = VertexPair.VertexID0;
+			const FVertexID SecondVertexID = VertexPair.VertexID1;
+
+			const int32 FirstVertexNumber = PerimeterVertexIDs.IndexOfByKey( FirstVertexID );
+			check( FirstVertexNumber != INDEX_NONE );	// Incoming vertex ID must already be a part of this polygon!
+			const int32 SecondVertexNumber = PerimeterVertexIDs.IndexOfByKey( SecondVertexID );
+			check( SecondVertexNumber != INDEX_NONE );	// Incoming vertex ID must already be a part of this polygon!
+
+
+			FPolygonToCreate& NewPolygon = *new( PolygonsToCreate ) FPolygonToCreate();
+			NewPolygon.PolygonGroupID = GetGroupForPolygon( PolygonID );
+			NewPolygon.PolygonEdgeHardness = EPolygonEdgeHardness::NewEdgesSoft;
+
+			static TArray<int32> PerimeterVertexNumbers;
+			PerimeterVertexNumbers.Reset();
+
+			const bool bWindsForward = FirstVertexNumber < SecondVertexNumber;
+
+			int32 SmallerVertexNumber = bWindsForward ? FirstVertexNumber : SecondVertexNumber;
+			int32 LargerVertexNumber = bWindsForward ? SecondVertexNumber : FirstVertexNumber;
+
+			if( PolygonIter == 0 || PolygonIter == ( NumPolygonsToCreate - 1 ) )
+			{
+				// This is either the first or last new polygon 
+				const bool bIsFirstPolygon = ( PolygonIter == 0 );
+
+				// Add the vertices we created for the new edge that will split the polygon
+				if( bIsFirstPolygon )
+				{
+					PerimeterVertexNumbers.Add( SmallerVertexNumber );
+					PerimeterVertexNumbers.Add( LargerVertexNumber );
+				}
+				else
+				{
+					PerimeterVertexNumbers.Add( LargerVertexNumber );
+					PerimeterVertexNumbers.Add( SmallerVertexNumber );
+				}
+
+				// Now add all of the other vertices of the original polygon that are on this side of the split
+				if( bIsFirstPolygon )
+				{
+					for( int32 VertexNumber = ( LargerVertexNumber + 1 ) % PerimeterVertexIDs.Num();
+						 VertexNumber != SmallerVertexNumber;
+						 VertexNumber = ( VertexNumber + 1 ) % PerimeterVertexIDs.Num() )
+					{
+						PerimeterVertexNumbers.Add( VertexNumber );
+					}
+				}
+				else
+				{
+					for( int32 VertexNumber = ( SmallerVertexNumber + 1 ) % PerimeterVertexIDs.Num();
+						 VertexNumber != LargerVertexNumber;
+						 VertexNumber = ( VertexNumber + 1 ) % PerimeterVertexIDs.Num() )
+					{
+						PerimeterVertexNumbers.Add( VertexNumber );
+					}
+				}
+			}
+			else
+			{
+				// @todo mesheditor urgent edgeloop: Polygon winding is incorrect with multiple splits
+
+				// This is a new polygon in the middle of other polygons created by the splits
+				PerimeterVertexNumbers.Add( bWindsForward ? SmallerVertexNumber : LargerVertexNumber );
+				PerimeterVertexNumbers.Add( bWindsForward ? LargerVertexNumber : SmallerVertexNumber );
+				PerimeterVertexNumbers.Add( LastPolygonVertexNumbers[ 1 ] );
+				PerimeterVertexNumbers.Add( LastPolygonVertexNumbers[ 0 ] );
+			}
+
+			NewPolygon.PerimeterVertices.Reserve( PerimeterVertexNumbers.Num() );
+			for( const int32 VertexNumber : PerimeterVertexNumbers )
+			{
+				FVertexAndAttributes& NewVertexAndAttributes = *new( NewPolygon.PerimeterVertices ) FVertexAndAttributes();
+				NewVertexAndAttributes.VertexID = FVertexID::Invalid;
+				NewVertexAndAttributes.VertexInstanceID = GetPolygonPerimeterVertexInstance( PolygonID, VertexNumber );
+			}
+
+			LastPolygonVertexNumbers[ 0 ] = PerimeterVertexNumbers[ 0 ];
+			LastPolygonVertexNumbers[ 1 ] = PerimeterVertexNumbers[ 1 ];
+			bLastPolygonWindsForward = bWindsForward;
+		}
+	}
+
+	// Delete the old polygons
+	{
+		static TArray<FPolygonID> PolygonIDsToDelete;
+		PolygonIDsToDelete.Reset();
+		PolygonIDsToDelete.Reserve( PolygonsToSplit.Num() );
+		for( const FPolygonToSplit& PolygonToSplit : PolygonsToSplit )
+		{
+			PolygonIDsToDelete.Add( PolygonToSplit.PolygonID );
+		}
+
+		const bool bDeleteOrphanEdges = false;
+		const bool bDeleteOrphanVertices = false;
+		const bool bDeleteOrphanVertexInstances = false;
+		const bool bDeleteEmptyPolygonGroups = false;
+		DeletePolygons( PolygonIDsToDelete, bDeleteOrphanEdges, bDeleteOrphanVertices, bDeleteOrphanVertexInstances, bDeleteEmptyPolygonGroups );
+	}
+
+	// Create new polygons that are split appropriately and connect to the new vertices we've added
+	static TArray< FPolygonID > NewPolygonIDs;
+	NewPolygonIDs.Reset();
+	static TArray< FEdgeID > NewEdgeIDs;
+	NewEdgeIDs.Reset();
+	CreatePolygons( PolygonsToCreate, /* Out */ NewPolygonIDs, /* Out */ NewEdgeIDs );
+
+	OutNewEdgeIDs.Append( NewEdgeIDs );
+}
+
+
+void UEditableMesh::DeleteEdgeAndConnectedPolygons( const FEdgeID EdgeID, const bool bDeleteOrphanedEdges, const bool bDeleteOrphanedVertices, const bool bDeleteOrphanedVertexInstances, const bool bDeleteEmptyPolygonGroups )
+{
+	static TArray<FPolygonID> PolygonIDsToDelete;
+	PolygonIDsToDelete.Reset();
+
+	const int32 ConnectedPolygonCount = GetEdgeConnectedPolygonCount( EdgeID );
+	for( int32 ConnectedPolygonNumber = 0; ConnectedPolygonNumber < ConnectedPolygonCount; ++ConnectedPolygonNumber )
+	{
+		const FPolygonID PolygonID = GetEdgeConnectedPolygon( EdgeID, ConnectedPolygonNumber );
+
+		// Although it can be uncommon, it's possible the edge is connecting the same polygon to itself.  We need to add uniquely.
+		PolygonIDsToDelete.AddUnique( PolygonID );
+	}
+
+	// Delete the polygons
+	DeletePolygons( PolygonIDsToDelete, bDeleteOrphanedEdges, bDeleteOrphanedVertices, bDeleteOrphanedVertexInstances, bDeleteEmptyPolygonGroups );
+
+	// If the caller asked us not deleted orphaned edges, our edge-to-delete will still be hanging around.  Let's go
+	// and delete it now.
+	if( !bDeleteOrphanedEdges )
+	{
+		// NOTE: Because we didn't delete any orphaned edges, the incoming edge ID should still be valid
+		static TArray<FEdgeID> EdgeIDsToDelete;
+		EdgeIDsToDelete.Reset();
+		EdgeIDsToDelete.Add( EdgeID );
+
+		// This edge MUST be an orphan!
+		check( GetEdgeConnectedPolygonCount( EdgeID ) == 0 );
+
+		DeleteEdges( EdgeIDsToDelete, bDeleteOrphanedVertices );
+	}
+}
+
+
+void UEditableMesh::DeleteVertexAndConnectedEdgesAndPolygons( const FVertexID VertexID, const bool bDeleteOrphanedEdges, const bool bDeleteOrphanedVertices, const bool bDeleteOrphanedVertexInstances, const bool bDeleteEmptyPolygonGroups )
+{
+	static TArray<FEdgeID> EdgeIDsToDelete;
+	EdgeIDsToDelete.Reset();
+
+	const int32 ConnectedEdgeCount = GetVertexConnectedEdgeCount( VertexID );
+	for( int32 ConnectedEdgeNumber = 0; ConnectedEdgeNumber < ConnectedEdgeCount; ++ConnectedEdgeNumber )
+	{
+		const FEdgeID ConnectedEdgeID = GetVertexConnectedEdge( VertexID, ConnectedEdgeNumber );
+		EdgeIDsToDelete.Add( ConnectedEdgeID );
+	}
+
+	for( const FEdgeID EdgeIDToDelete : EdgeIDsToDelete )
+	{
+		// Make sure the edge still exists.  It may have been deleted as a polygon's edges were deleted during
+		// a previous iteration through this loop.
+		if( IsValidEdge( EdgeIDToDelete ) )
+		{
+			DeleteEdgeAndConnectedPolygons( EdgeIDToDelete, bDeleteOrphanedEdges, bDeleteOrphanedVertices, bDeleteOrphanedVertexInstances, bDeleteEmptyPolygonGroups );
+		}
+	}
+}
+
+
+void UEditableMesh::DeleteOrphanVertices( const TArray<FVertexID>& VertexIDsToDelete )
+{
+	UE_LOG( LogEditableMesh, Verbose, TEXT( "DeleteOrphanVertices: %s" ), *LogHelpers::ArrayToString( VertexIDsToDelete ) );
+
+	// Back everything up
+	{
+		FCreateVerticesChangeInput RevertInput;
+		RevertInput.VerticesToCreate.Reserve( VertexIDsToDelete.Num() );
+
+		// NOTE: We iterate backwards, to restore vertices in the opposite order that we deleted them
+		for( int32 VertexNumber = VertexIDsToDelete.Num() - 1; VertexNumber >= 0; --VertexNumber )
+		{
+			const FVertexID VertexID = VertexIDsToDelete[ VertexNumber ];
+
+			// Make sure the vertex is truly an orphan.  We're not going to be able to restore it's polygon vertex attributes,
+			// because the polygons won't exist when we're restoring the change
+			const FMeshVertex& Vertex = Vertices[ VertexID.GetValue() ];
+			check( Vertex.ConnectedEdgeIDs.Num() == 0 );
+			check( Vertex.VertexInstanceIDs.Num() == 0 );
+
+			RevertInput.VerticesToCreate.Emplace();
+			FVertexToCreate& VertexToCreate = RevertInput.VerticesToCreate.Last();
+
+			VertexToCreate.OriginalVertexID = VertexID;
+
+			for( const FName AttributeName : UEditableMesh::GetValidVertexAttributes() )
+			{
+				const int32 MaxAttributeIndex = GetMaxAttributeIndex( AttributeName );
+				for( int32 AttributeIndex = 0; AttributeIndex < MaxAttributeIndex; ++AttributeIndex )
+				{
+					VertexToCreate.VertexAttributes.Attributes.Emplace(
+						AttributeName,
+						AttributeIndex,
+						GetVertexAttribute( VertexID, AttributeName, AttributeIndex )
+					);
+				}
+			}
+		}
+
+		AddUndo( MakeUnique<FCreateVerticesChange>( MoveTemp( RevertInput ) ) );
+	}
+
+	// Give the adapter a chance to handle this before they are deleted
+	for( UEditableMeshAdapter* Adapter : Adapters )
+	{
+		Adapter->OnDeleteOrphanVertices( this, VertexIDsToDelete );
+	}
+	
+	// Actually delete the vertices
+	{
+		for( const FVertexID VertexIDToDelete : VertexIDsToDelete )
+		{
+			Vertices.RemoveAt( VertexIDToDelete.GetValue() );
+		}
+	}
+
+	UE_LOG( LogEditableMesh, Verbose, TEXT( "* DeleteOrphanVertices returned" ) );
+}
+
+
+void UEditableMesh::DeleteVertexInstances( const TArray<FVertexInstanceID>& VertexInstanceIDsToDelete, const bool bDeleteOrphanedVertices )
+{
+	UE_LOG( LogEditableMesh, Verbose, TEXT( "DeleteVertexInstances: %s" ), *LogHelpers::ArrayToString( VertexInstanceIDsToDelete ) );
+
+	// Back everything up
+	{
+		FCreateVertexInstancesChangeInput RevertInput;
+		RevertInput.VertexInstancesToCreate.Reserve( VertexInstanceIDsToDelete.Num() );
+
+		// NOTE: We iterate backwards, to restore vertices in the opposite order that we deleted them
+		for( int32 VertexNumber = VertexInstanceIDsToDelete.Num() - 1; VertexNumber >= 0; --VertexNumber )
+		{
+			const FVertexInstanceID VertexInstanceID = VertexInstanceIDsToDelete[ VertexNumber ];
+			const FMeshVertexInstance& VertexInstance = VertexInstances[ VertexInstanceID.GetValue() ];
+
+			// Back up properties
+			RevertInput.VertexInstancesToCreate.Emplace();
+			FVertexInstanceToCreate& VertexInstanceToCreate = RevertInput.VertexInstancesToCreate.Last();
+
+			VertexInstanceToCreate.VertexID = VertexInstance.VertexID;
+			VertexInstanceToCreate.OriginalVertexInstanceID = VertexInstanceID;
+
+			for( const FName AttributeName : GetValidVertexInstanceAttributes() )
+			{
+				const int32 MaxAttributeIndex = GetMaxAttributeIndex( AttributeName );
+				for( int32 AttributeIndex = 0; AttributeIndex < MaxAttributeIndex; ++AttributeIndex )
+				{
+					VertexInstanceToCreate.VertexInstanceAttributes.Attributes.Emplace(
+						AttributeName,
+						AttributeIndex,
+						GetVertexInstanceAttribute( VertexInstanceID, AttributeName, AttributeIndex )
+					);
+				}
+			}
+		}
+
+		AddUndo( MakeUnique<FCreateVertexInstancesChange>( MoveTemp( RevertInput ) ) );
+	}
+
+	// Give the adapter a chance to do something with this event before it happens
+	for( UEditableMeshAdapter* Adapter : Adapters )
+	{
+		Adapter->OnDeleteVertexInstances( this, VertexInstanceIDsToDelete );
+	}
+
+	static TArray<FVertexID> OrphanedVertexIDs;
+	OrphanedVertexIDs.Reset();
+
+	// Actually delete the vertex instances
+	{
+		for( const FVertexInstanceID VertexInstanceIDToDelete : VertexInstanceIDsToDelete )
+		{
+			// Check that no polygon is referencing this vertex instance
+			check( GetVertexInstanceConnectedPolygonCount( VertexInstanceIDToDelete ) == 0 );
+
+			// Remove the reference in the vertex it instances
+			const FMeshVertexInstance& VertexInstance = VertexInstances[ VertexInstanceIDToDelete.GetValue() ];
+			FMeshVertex& Vertex = Vertices[ VertexInstance.VertexID.GetValue() ];
+			verify( Vertex.VertexInstanceIDs.Remove( VertexInstanceIDToDelete ) == 1 );
+
+			if( bDeleteOrphanedVertices )
+			{
+				if( Vertex.VertexInstanceIDs.Num() == 0 )
+				{
+					OrphanedVertexIDs.Add( VertexInstance.VertexID );
+				}
+			}
+
+			// Delete the element from the container
+			VertexInstances.RemoveAt( VertexInstanceIDToDelete.GetValue() );
+		}
+	}
+
+	// Delete orphaned vertices, if there are any.
+	if( OrphanedVertexIDs.Num() > 0 )
+	{
+		DeleteOrphanVertices( OrphanedVertexIDs );
+	}
+
+	UE_LOG( LogEditableMesh, Verbose, TEXT( "* DeleteVertexInstances returned" ) );
+}
+
+
+void UEditableMesh::DeleteEdges( const TArray<FEdgeID>& EdgeIDsToDelete, const bool bDeleteOrphanedVertices )
+{
+	UE_LOG( LogEditableMesh, Verbose, TEXT( "DeleteEdges: %s" ), *LogHelpers::ArrayToString( EdgeIDsToDelete ) );
+
+	// Back everything up
+	{
+		FCreateEdgesChangeInput RevertInput;
+
+		// NOTE: We iterate backwards, to restore edges in the opposite order that we deleted them
+		for( int32 EdgeNumber = EdgeIDsToDelete.Num() - 1; EdgeNumber >= 0; --EdgeNumber )
+		{
+			const FEdgeID EdgeID = EdgeIDsToDelete[ EdgeNumber ];
+			const FMeshEdge& Edge = Edges[ EdgeID.GetValue() ];
+
+			RevertInput.EdgesToCreate.Emplace();
+			FEdgeToCreate& EdgeToCreate = RevertInput.EdgesToCreate.Last();
+
+			EdgeToCreate.OriginalEdgeID = EdgeID;
+			EdgeToCreate.VertexID0 = Edge.VertexIDs[ 0 ];
+			EdgeToCreate.VertexID1 = Edge.VertexIDs[ 1 ];
+			EdgeToCreate.ConnectedPolygons = Edge.ConnectedPolygons;
+			
+			for( const FName AttributeName : UEditableMesh::GetValidEdgeAttributes() )
+			{
+				const int32 MaxAttributeIndex = GetMaxAttributeIndex( AttributeName );
+				for( int32 AttributeIndex = 0; AttributeIndex < MaxAttributeIndex; ++AttributeIndex )
+				{
+					EdgeToCreate.EdgeAttributes.Attributes.Emplace(
+						AttributeName,
+						AttributeIndex,
+						GetEdgeAttribute( EdgeID, AttributeName, AttributeIndex )
+					);
+				}
+			}
+		}
+
+		AddUndo( MakeUnique<FCreateEdgesChange>( MoveTemp( RevertInput ) ) );
+	}
+
+	for( UEditableMeshAdapter* Adapter : Adapters )
+	{
+		Adapter->OnDeleteEdges( this, EdgeIDsToDelete );
+	}
+
+	// Delete the edges
+	{
+		static TArray<FVertexID> OrphanedVertexIDs;
+		OrphanedVertexIDs.Reset();
+
+		for( const FEdgeID EdgeID : EdgeIDsToDelete )
+		{
+			FMeshEdge& Edge = Edges[ EdgeID.GetValue() ];
+
+			for( const FVertexID EdgeVertexID : Edge.VertexIDs )
+			{
+				FMeshVertex& Vertex = Vertices[ EdgeVertexID.GetValue() ];
+				verify( Vertex.ConnectedEdgeIDs.RemoveSingle( EdgeID ) == 1 );
+
+				// If the vertex has no more edges connected, we'll keep track of that so we can delete the vertex later
+				if( bDeleteOrphanedVertices && Vertex.ConnectedEdgeIDs.Num() == 0 )
+				{
+					check( Vertex.VertexInstanceIDs.Num() == 0 );	// We must already have deleted any vertex instances
+					check( !OrphanedVertexIDs.Contains( EdgeVertexID ) );	// Orphaned vertex shouldn't have already been orphaned by an earlier deleted edge passed into this function
+					OrphanedVertexIDs.Add( EdgeVertexID );
+				}
+			}
+
+			// Delete the edge
+			Edges.RemoveAt( EdgeID.GetValue() );
+		}
+
+		// If we orphaned any vertices and we were asked to delete those, then we'll go ahead and do that now.
+		if( OrphanedVertexIDs.Num() > 0 )
+		{
+			DeleteOrphanVertices( OrphanedVertexIDs );
+		}
+	}
+
+	UE_LOG( LogEditableMesh, Verbose, TEXT( "* DeleteEdges returned" ) );
+}
+
+
+void UEditableMesh::CreateEmptyVertexRange( const int32 NumVerticesToCreate, TArray<FVertexID>& OutNewVertexIDs )
+{
+	UE_LOG( LogEditableMesh, Verbose, TEXT( "CreateEmptyVertexRange: %d" ), NumVerticesToCreate );
+
+	OutNewVertexIDs.Reset( NumVerticesToCreate );
+
+	// Reserve elements
+	Vertices.Reserve( Vertices.Num() + NumVerticesToCreate );
+
+	for( int32 Count = 0; Count < NumVerticesToCreate; ++Count )
+	{
+		// Allocate vertex
+		const FVertexID VertexID = FVertexID( Vertices.Add( FMeshVertex() ) );
+		FMeshVertex& Vertex = Vertices[ VertexID.GetValue() ];
+
+		// Initialize vertex
+		Vertex.VertexPosition = FVector::ZeroVector;
+		Vertex.CornerSharpness = 0.0f;
+
+		OutNewVertexIDs.Add( VertexID );
+	}
+
+	// NOTE: We iterate backwards, to delete vertices in the opposite order that we added them
+	{
+		FDeleteOrphanVerticesChangeInput RevertInput;
+		RevertInput.VertexIDsToDelete.Reserve( NumVerticesToCreate );
+		for( int32 VertexNumber = OutNewVertexIDs.Num() - 1; VertexNumber >= 0; --VertexNumber )
+		{
+			RevertInput.VertexIDsToDelete.Add( OutNewVertexIDs[ VertexNumber ] );
+		}
+
+		AddUndo( MakeUnique<FDeleteOrphanVerticesChange>( MoveTemp( RevertInput ) ) );
+	}
+
+	// Advise the adapter of new vertices
+	for( UEditableMeshAdapter* Adapter : Adapters )
+	{
+		Adapter->OnCreateEmptyVertexRange( this, OutNewVertexIDs );
+	}
+
+	UE_LOG( LogEditableMesh, Verbose, TEXT( "* CreateEmptyVertexRange returned %s" ), *LogHelpers::ArrayToString( OutNewVertexIDs ) );
+}
+
+
+void UEditableMesh::CreateVertices( const TArray<FVertexToCreate>& VerticesToCreate, TArray<FVertexID>& OutNewVertexIDs )
+{
+	UE_LOG( LogEditableMesh, Verbose, TEXT( "CreateVertices: %s" ), *LogHelpers::ArrayToString( VerticesToCreate ) );
+
+	OutNewVertexIDs.Reset( VerticesToCreate.Num() );
+
+	// Reserve elements
+	Vertices.Reserve( Vertices.Num() + VerticesToCreate.Num() );
+
+	for( const FVertexToCreate& VertexToCreate : VerticesToCreate )
+	{
+		// Allocate vertex
+		FVertexID VertexID = VertexToCreate.OriginalVertexID;
+		if( VertexID != FVertexID::Invalid )
+		{
+			Vertices.Insert( VertexID.GetValue(), FMeshVertex() );
+		}
+		else
+		{
+			VertexID = FVertexID( Vertices.Add( FMeshVertex() ) );
+		}
+
+		FMeshVertex& Vertex = Vertices[ VertexID.GetValue() ];
+
+		// Initialize vertex
+		Vertex.VertexPosition = FVector::ZeroVector;
+		Vertex.CornerSharpness = 0.0f;
+
+		OutNewVertexIDs.Add( VertexID );
+	}
+
+	// NOTE: We iterate backwards, to delete vertices in the opposite order that we added them
+	{
+		FDeleteOrphanVerticesChangeInput RevertInput;
+		RevertInput.VertexIDsToDelete.Reserve( VerticesToCreate.Num() );
+		for( int32 VertexNumber = OutNewVertexIDs.Num() - 1; VertexNumber >= 0; --VertexNumber )
+		{
+			RevertInput.VertexIDsToDelete.Add( OutNewVertexIDs[ VertexNumber ] );
+		}
+
+		AddUndo( MakeUnique<FDeleteOrphanVerticesChange>( MoveTemp( RevertInput ) ) );
+	}
+
+	// Advise the adapter of new vertices
+	for( UEditableMeshAdapter* Adapter : Adapters )
+	{
+		Adapter->OnCreateVertices( this, OutNewVertexIDs );
+	}
+
+	// Set new vertex attributes
+	for( int32 Index = 0; Index < OutNewVertexIDs.Num(); ++Index )
+	{
+		for( const FMeshElementAttributeData& VertexAttribute : VerticesToCreate[ Index ].VertexAttributes.Attributes )
+		{
+			SetVertexAttribute( OutNewVertexIDs[ Index ], VertexAttribute.AttributeName, VertexAttribute.AttributeIndex, VertexAttribute.AttributeValue );
+		}
+	}
+
+	UE_LOG( LogEditableMesh, Verbose, TEXT( "* CreateVertices returned %s" ), *LogHelpers::ArrayToString( OutNewVertexIDs ) );
+}
+
+
+void UEditableMesh::CreateVertexInstances( const TArray<FVertexInstanceToCreate>& VertexInstancesToCreate, TArray<FVertexInstanceID>& OutNewVertexInstanceIDs )
+{
+	UE_LOG( LogEditableMesh, Verbose, TEXT( "CreateVertexInstances: %s" ), *LogHelpers::ArrayToString( VertexInstancesToCreate ) );
+
+	OutNewVertexInstanceIDs.Reset( VertexInstancesToCreate.Num() );
+
+	// Reserve elements
+	VertexInstances.Reserve( VertexInstances.Num() + VertexInstancesToCreate.Num() );
+
+	for( const FVertexInstanceToCreate& VertexInstanceToCreate : VertexInstancesToCreate )
+	{
+		// Allocate vertex instance
+		FVertexInstanceID VertexInstanceID = VertexInstanceToCreate.OriginalVertexInstanceID;
+		if( VertexInstanceID != FVertexInstanceID::Invalid )
+		{
+			VertexInstances.Insert( VertexInstanceID.GetValue(), FMeshVertexInstance() );
+		}
+		else
+		{
+			VertexInstanceID = FVertexInstanceID( VertexInstances.Add( FMeshVertexInstance() ) );
+		}
+
+		FMeshVertexInstance& VertexInstance = VertexInstances[ VertexInstanceID.GetValue() ];
+
+		// Initialize vertex instance
+		VertexInstance.VertexID = VertexInstanceToCreate.VertexID;
+
+		VertexInstance.Normal = FVector::ZeroVector;
+		VertexInstance.Tangent = FVector::ZeroVector;
+		VertexInstance.BinormalSign = 0.0f;
+		VertexInstance.Color = FLinearColor::White;
+		VertexInstance.VertexUVs.SetNumZeroed( TextureCoordinateCount );
+
+		// Add to the list of instances in the vertex
+		check( !Vertices[ VertexInstance.VertexID.GetValue() ].VertexInstanceIDs.Contains( VertexInstanceID ) );
+		Vertices[ VertexInstance.VertexID.GetValue() ].VertexInstanceIDs.Add( VertexInstanceID );
+
+		OutNewVertexInstanceIDs.Add( VertexInstanceID );
+	}
+
+	// NOTE: We iterate backwards, to delete vertex instances in the opposite order that we added them
+	{
+		FDeleteVertexInstancesChangeInput RevertInput;
+		RevertInput.bDeleteOrphanedVertices = false;
+		RevertInput.VertexInstanceIDsToDelete.Reserve( OutNewVertexInstanceIDs.Num() );
+		for( int32 VertexNumber = OutNewVertexInstanceIDs.Num() - 1; VertexNumber >= 0; --VertexNumber )
+		{
+			RevertInput.VertexInstanceIDsToDelete.Add( OutNewVertexInstanceIDs[ VertexNumber ] );
+		}
+
+		AddUndo( MakeUnique<FDeleteVertexInstancesChange>( MoveTemp( RevertInput ) ) );
+	}
+
+	// Advise the adapter of new vertex instances
+	for( UEditableMeshAdapter* Adapter : Adapters )
+	{
+		Adapter->OnCreateVertexInstances( this, OutNewVertexInstanceIDs );
+	}
+
+	// Now set vertex instance attributes
+	for( int32 Index = 0; Index < OutNewVertexInstanceIDs.Num(); ++Index )
+	{
+		for( const FMeshElementAttributeData& VertexInstanceAttribute : VertexInstancesToCreate[ Index ].VertexInstanceAttributes.Attributes )
+		{
+			SetVertexInstanceAttribute( OutNewVertexInstanceIDs[ Index ], VertexInstanceAttribute.AttributeName, VertexInstanceAttribute.AttributeIndex, VertexInstanceAttribute.AttributeValue );
+		}
+	}
+
+	UE_LOG( LogEditableMesh, Verbose, TEXT( "* CreateVertexInstances returned %s" ), *LogHelpers::ArrayToString( OutNewVertexInstanceIDs ) );
+}
+
+
+void UEditableMesh::CreateEdges( const TArray<FEdgeToCreate>& EdgesToCreate, TArray<FEdgeID>& OutNewEdgeIDs )
+{
+	UE_LOG( LogEditableMesh, Verbose, TEXT( "CreateEdges: %s" ), *LogHelpers::ArrayToString( EdgesToCreate ) );
+
+	OutNewEdgeIDs.Reset( EdgesToCreate.Num() );
+
+	// Reserve elements
+	Edges.Reserve( Edges.Num() + EdgesToCreate.Num() );
+
+	for( const FEdgeToCreate& EdgeToCreate : EdgesToCreate )
+	{
+		// Allocate edge
+		FEdgeID EdgeID = EdgeToCreate.OriginalEdgeID;
+		if( EdgeID != FEdgeID::Invalid )
+		{
+			Edges.Insert( EdgeID.GetValue(), FMeshEdge() );
+		}
+		else
+		{
+			EdgeID = FEdgeID( Edges.Add( FMeshEdge() ) );
+		}
+
+		FMeshEdge& Edge = Edges[ EdgeID.GetValue() ];
+
+		// Initialize edge
+		Edge.VertexIDs[ 0 ] = EdgeToCreate.VertexID0;
+		Edge.VertexIDs[ 1 ] = EdgeToCreate.VertexID1;
+		Edge.ConnectedPolygons = EdgeToCreate.ConnectedPolygons;
+
+		Edge.bIsHardEdge = false;
+		Edge.CreaseSharpness = 0.0f;
+
+		// Connect the edge to its vertices
+		Vertices[ Edge.VertexIDs[ 0 ].GetValue() ].ConnectedEdgeIDs.AddUnique( EdgeID );
+		Vertices[ Edge.VertexIDs[ 1 ].GetValue() ].ConnectedEdgeIDs.AddUnique( EdgeID );
+
+		OutNewEdgeIDs.Add( EdgeID );
+	}
+
+	// NOTE: We iterate backwards, to delete edges in the opposite order that we added them
+	{
+		FDeleteEdgesChangeInput RevertInput;
+		RevertInput.bDeleteOrphanedVertices = false;	// Don't delete any vertices on revert.  We're only creating edges here, not vertices!
+		RevertInput.EdgeIDsToDelete.Reserve( EdgesToCreate.Num() );
+		for( int32 EdgeNumber = OutNewEdgeIDs.Num() - 1; EdgeNumber >= 0; --EdgeNumber )
+		{
+			RevertInput.EdgeIDsToDelete.Add( OutNewEdgeIDs[ EdgeNumber ] );
+		}
+
+		AddUndo( MakeUnique<FDeleteEdgesChange>( MoveTemp( RevertInput ) ) );
+	}
+
+	// Advise the adapter that edges have been created
+	for( UEditableMeshAdapter* Adapter : Adapters )
+	{
+		Adapter->OnCreateEdges( this, OutNewEdgeIDs );
+	}
+
+	for( int32 Index = 0; Index < OutNewEdgeIDs.Num(); ++Index )
+	{
+		for( const FMeshElementAttributeData& EdgeAttribute : EdgesToCreate[ Index ].EdgeAttributes.Attributes )
+		{
+			const bool bIsUndo = false;
+			SetEdgeAttribute( OutNewEdgeIDs[ Index ], EdgeAttribute.AttributeName, EdgeAttribute.AttributeIndex, EdgeAttribute.AttributeValue, bIsUndo );
+		}
+	}
+
+	UE_LOG( LogEditableMesh, Verbose, TEXT( "* CreateEdges returned %s" ), *LogHelpers::ArrayToString( OutNewEdgeIDs ) );
+}
+
+
+FVertexInstanceID UEditableMesh::CreateVertexInstanceForContourVertex( const FVertexAndAttributes& ContourVertex, const FPolygonID PolygonID )
+{
+	FVertexInstanceID NewVertexInstanceID = ContourVertex.VertexInstanceID;
+
+	// Determine whether we need to create a new vertex instance or not. We need to if:
+	//   a) no valid vertex instance ID was supplied;
+	// and then:
+	//   b) if there are no vertex instances in the same soft edged polygon group as the new polygon which share the attributes being assigned
+	if( NewVertexInstanceID == FVertexInstanceID::Invalid )
+	{
+		const FVertexID VertexID = ContourVertex.VertexID;
+		check( VertexID != FVertexID::Invalid );
+
+		// Find candidate vertex instances in the same smoothing group as this polygon
+		static TArray<FVertexInstanceID> PossibleVertexInstances;
+		PossibleVertexInstances.Reset();
+
+		const bool bPolygonNotYetInitialized = true;
+		GetVertexInstancesInSameSoftEdgedGroup( VertexID, PolygonID, bPolygonNotYetInitialized, PossibleVertexInstances );
+
+		// Iterate through each potential vertex instance, checking if its attributes are equal before and after the new attributes are applied.
+		// If so, we can merge with this vertex.
+		// Note: unspecified attributes in the FVertexAndAttributes struct are "don't cares".
+		for( const FVertexInstanceID PossibleVertexInstance : PossibleVertexInstances )
+		{
+			// Get mergeable attribute data for this vertex instance
+			static TArray<FMeshElementAttributeData> AllAttributes;
+			GetMergedNamedVertexInstanceAttributeData( GetMergeableVertexInstanceAttributes(), PossibleVertexInstance, TArray<FMeshElementAttributeData>(), AllAttributes );
+
+			// Get attribute data for this vertex instance as it would be after applying the new attributes
+			static TArray<FMeshElementAttributeData> AllAttributesModified;
+			GetMergedNamedVertexInstanceAttributeData( GetMergeableVertexInstanceAttributes(), PossibleVertexInstance, ContourVertex.PolygonVertexAttributes.Attributes, AllAttributesModified );
+
+			// If it's the same, we can use this vertex instance instead of creating a new one
+			if( AreAttributeListsEqual( AllAttributes, AllAttributesModified ) )
+			{
+				NewVertexInstanceID = PossibleVertexInstance;
+				break;
+			}
+		}
+
+		// If we didn't find any suitable existing vertex instance to merge with, create a new one.
+		if( NewVertexInstanceID == FVertexInstanceID::Invalid )
+		{
+			FVertexInstanceToCreate VertexInstanceToCreate;
+			VertexInstanceToCreate.VertexID = VertexID;
+			VertexInstanceToCreate.VertexInstanceAttributes = ContourVertex.PolygonVertexAttributes;
+
+			static TArray<FVertexInstanceToCreate> VertexInstancesToCreate;
+			VertexInstancesToCreate.Reset();
+			VertexInstancesToCreate.Add( VertexInstanceToCreate );
+
+			static TArray<FVertexInstanceID> CreatedVertexInstances;
+			CreateVertexInstances( VertexInstancesToCreate, CreatedVertexInstances );
+			check( CreatedVertexInstances.Num() == 1 );
+			NewVertexInstanceID = CreatedVertexInstances[ 0 ];
+		}
+	}
+	else
+	{
+		// Cannot specify both a vertex instance ID and a vertex ID
+		check( ContourVertex.VertexID == FVertexID::Invalid );
+	}
+
+	return NewVertexInstanceID;
+}
+
+
+void UEditableMesh::CreatePolygonContour( const FPolygonID PolygonID, const TArray<FVertexAndAttributes>& Contour, const EPolygonEdgeHardness PolygonEdgeHardness, TArray<FEdgeID>& OutEdgeIDs, TArray<FVertexInstanceID>& OutVertexInstanceIDs )
+{
+	// All polygons must have at least three vertices
+	const int32 NumContourVertices = Contour.Num();
+	check( NumContourVertices >= 3 );
+
+	OutEdgeIDs.Reset( NumContourVertices );
+	OutVertexInstanceIDs.SetNumUninitialized( NumContourVertices );
+
+	static TArray<FEdgeID> AllEdges;
+	AllEdges.SetNumUninitialized( NumContourVertices );
+
+	// Get or create edges
+	for( int32 VertexNumber = 0; VertexNumber < NumContourVertices; ++VertexNumber )
+	{
+		const int32 NextVertexNumber = ( VertexNumber + 1 ) % NumContourVertices;
+
+		// Get vertex ID for the ends of the edge.
+		FVertexID VertexID0 = Contour[ VertexNumber ].VertexID;
+		FVertexID VertexID1 = Contour[ NextVertexNumber ].VertexID;
+
+		// If valid vertex IDs weren't supplied, we'll deduce it from the vertex instance IDs instead (which must be valid instead).
+		if( VertexID0 == FVertexID::Invalid )
+		{
+			const FVertexInstanceID VertexInstanceID0 = Contour[ VertexNumber ].VertexInstanceID;
+			check( VertexInstanceID0 != FVertexInstanceID::Invalid );
+			VertexID0 = VertexInstances[ VertexInstanceID0.GetValue() ].VertexID;
+		}
+
+		if( VertexID1 == FVertexID::Invalid )
+		{
+			const FVertexInstanceID VertexInstanceID1 = Contour[ NextVertexNumber ].VertexInstanceID;
+			check( VertexInstanceID1 != FVertexInstanceID::Invalid );
+			VertexID1 = VertexInstances[ VertexInstanceID1.GetValue() ].VertexID;
+		}
+
+		bool bOutEdgeWindingIsReversed;
+		FEdgeID EdgeID = GetVertexPairEdge( VertexID0, VertexID1, bOutEdgeWindingIsReversed );
+		if( EdgeID == FEdgeID::Invalid )
+		{
+			// Create the new edge!  Note that this does not connect the edge to the polygon.  We expect the caller to do that afterwards.
+			FEdgeToCreate EdgeToCreate;
+			EdgeToCreate.VertexID0 = VertexID0;
+			EdgeToCreate.VertexID1 = VertexID1;
+
+			if( PolygonEdgeHardness == EPolygonEdgeHardness::NewEdgesHard )
+			{
+				EdgeToCreate.EdgeAttributes.Attributes.Emplace( UEditableMeshAttribute::EdgeIsHard(), 0, FVector4( 1.0f ) );
+			}
+
+			static TArray<FEdgeToCreate> EdgesToCreate;
+			EdgesToCreate.Reset();
+			EdgesToCreate.Add( EdgeToCreate );
+
+			static TArray<FEdgeID> NewEdgeIDs;
+			CreateEdges( EdgesToCreate, /* Out */ NewEdgeIDs );
+			check( NewEdgeIDs.Num() == 1 );
+			EdgeID = NewEdgeIDs[ 0 ];
+
+			OutEdgeIDs.Add( EdgeID );
+		}
+
+		// Connect the edge to the polygon
+		FMeshEdge& Edge = Edges[ EdgeID.GetValue() ];
+		Edge.ConnectedPolygons.Add( PolygonID );
+
+		AllEdges[ VertexNumber ] = EdgeID;
+	}
+
+	if( PolygonEdgeHardness == EPolygonEdgeHardness::AllEdgesHard ||
+		PolygonEdgeHardness == EPolygonEdgeHardness::AllEdgesSoft )
+	{
+		static TArray<FAttributesForEdge> AttributesForEdges;
+		AttributesForEdges.Reset( AllEdges.Num() );
+
+		for( const FEdgeID EdgeID : AllEdges )
+		{
+			AttributesForEdges.Emplace();
+			FAttributesForEdge& AttributesForEdge = AttributesForEdges.Last();
+			AttributesForEdge.EdgeID = EdgeID;
+			AttributesForEdge.EdgeAttributes.Attributes.Emplace( UEditableMeshAttribute::EdgeIsHard(), 0, FVector4( PolygonEdgeHardness == EPolygonEdgeHardness::AllEdgesHard ? 1.0f : 0.0f ) );
+		}
+
+		SetEdgesAttributes( AttributesForEdges );
+	}
+
+	// Get or create new vertex instances for this polygon.
+	// Do this after creating all the edges because:
+	//  1) we can establish a smoothing group for the new polygon based on adjacent edge hardness.
+	//  2) the revert operation will first delete all vertex instances, and then the edges (which can only delete orphaned vertices if they have no instances).
+	for( int32 VertexNumber = 0; VertexNumber < NumContourVertices; ++VertexNumber )
+	{
+		const FVertexInstanceID ContourVertexInstanceID = CreateVertexInstanceForContourVertex( Contour[ VertexNumber ], PolygonID );
+		OutVertexInstanceIDs[ VertexNumber ] = ContourVertexInstanceID;
+
+		// Connect this vertex instance to the polygon
+		FMeshVertexInstance& ContourVertexInstance = VertexInstances[ ContourVertexInstanceID.GetValue() ];
+		check( !ContourVertexInstance.ConnectedPolygons.Contains( PolygonID ) );
+		ContourVertexInstance.ConnectedPolygons.Add( PolygonID );
+	}
+}
+
+
+void UEditableMesh::CreatePolygons( const TArray<FPolygonToCreate>& PolygonsToCreate, TArray<FPolygonID>& OutNewPolygonIDs, TArray<FEdgeID>& OutNewEdgeIDs )
+{
+	UE_LOG( LogEditableMesh, Verbose, TEXT( "CreatePolygons: %s" ), *LogHelpers::ArrayToString( PolygonsToCreate ) );
+
+	OutNewPolygonIDs.Reset( PolygonsToCreate.Num() );
+	OutNewEdgeIDs.Reset();
+
+	Polygons.Reserve( Polygons.Num() + PolygonsToCreate.Num() );
+
+	for( const FPolygonToCreate& PolygonToCreate : PolygonsToCreate )
+	{
+		// Allocate polygon
+		FPolygonID PolygonID = PolygonToCreate.OriginalPolygonID;
+		if( PolygonID != FPolygonID::Invalid )
+		{
+			Polygons.Insert( PolygonID.GetValue(), FMeshPolygon() );
+		}
+		else
+		{
+			PolygonID = FPolygonID( Polygons.Add( FMeshPolygon() ) );
+		}
+
+		// Connect polygon to group and vice versa
+		FMeshPolygon& Polygon = Polygons[ PolygonID.GetValue() ];
+		FMeshPolygonGroup& PolygonGroup = PolygonGroups[ PolygonToCreate.PolygonGroupID.GetValue() ];
+		Polygon.PolygonGroupID = PolygonToCreate.PolygonGroupID;
+		PolygonGroup.Polygons.Add( PolygonID );
+
+		OutNewPolygonIDs.Add( PolygonID );
+
+		// Create edges and vertex instances for perimeter contour
+		static TArray<FEdgeID> PerimeterEdgeIDs;
+		CreatePolygonContour( PolygonID, PolygonToCreate.PerimeterVertices, PolygonToCreate.PolygonEdgeHardness, PerimeterEdgeIDs, Polygon.PerimeterContour.VertexInstanceIDs );
+		OutNewEdgeIDs.Append( PerimeterEdgeIDs );
+
+		// Iterate through hole definitions, and create edges and vertex instances for those
+		for( const FPolygonHoleVertices& PolygonHole : PolygonToCreate.PolygonHoles )
+		{
+			Polygon.HoleContours.Emplace();
+			FMeshPolygonContour& Contour = Polygon.HoleContours.Last();
+
+			static TArray<FEdgeID> HoleEdgeIDs;
+			CreatePolygonContour( PolygonID, PolygonHole.HoleVertices, EPolygonEdgeHardness::AllEdgesHard, HoleEdgeIDs, Contour.VertexInstanceIDs );
+			OutNewEdgeIDs.Append( HoleEdgeIDs );
+		}
+	}
+
+	for( UEditableMeshAdapter* Adapter : Adapters )
+	{
+		Adapter->OnCreatePolygons( this, OutNewPolygonIDs );
+	}
+
+	// Generate tangent basis for the polygon
+	PolygonsPendingNewTangentBasis.Append( OutNewPolygonIDs );
+
+	// Generate triangles for the new polygon
+	PolygonsPendingTriangulation.Append( OutNewPolygonIDs );
+
+	// NOTE: We iterate backwards, to delete polygons in the opposite order that we added them
+	{
+		FDeletePolygonsChangeInput RevertInput;
+		RevertInput.PolygonIDsToDelete.Reserve( PolygonsToCreate.Num() );
+		RevertInput.bDeleteOrphanedEdges = false;	// Don't delete any edges on revert.  We're only creating polygons here, not vertices!
+		RevertInput.bDeleteOrphanedVertices = false;	// Don't delete any vertices on revert.  We're only creating polygons here, not vertices!
+		RevertInput.bDeleteOrphanedVertexInstances = false;
+		RevertInput.bDeleteEmptySections = false;
+		for( int32 PolygonNumber = OutNewPolygonIDs.Num() - 1; PolygonNumber >= 0; --PolygonNumber )
+		{
+			RevertInput.PolygonIDsToDelete.Add( OutNewPolygonIDs[ PolygonNumber ] );
+		}
+
+		AddUndo( MakeUnique<FDeletePolygonsChange>( MoveTemp( RevertInput ) ) );
+	}
+
+	UE_LOG( LogEditableMesh, Verbose, TEXT( "* CreatePolygons returned %s, %s" ), *LogHelpers::ArrayToString( OutNewPolygonIDs ), *LogHelpers::ArrayToString( OutNewEdgeIDs ) );
+}
+
+
+void UEditableMesh::BackupPolygonContour( const FMeshPolygonContour& Contour, TArray<FVertexAndAttributes>& OutVerticesAndAttributes )
+{
+	OutVerticesAndAttributes.Reserve( Contour.VertexInstanceIDs.Num() );
+	for( const FVertexInstanceID VertexInstanceID : Contour.VertexInstanceIDs )
+	{
+		const FMeshVertexInstance& VertexInstance = VertexInstances[ VertexInstanceID.GetValue() ];
+
+		OutVerticesAndAttributes.Emplace();
+		FVertexAndAttributes& VertexAndAttributes = OutVerticesAndAttributes.Last();
+
+		// We rely on undoing recreating vertex instances, therefore we only need pass their IDs.
+		VertexAndAttributes.VertexInstanceID = VertexInstanceID;
+		VertexAndAttributes.VertexID = FVertexID::Invalid;
+	}
+}
+
+
+void UEditableMesh::DeletePolygonContour( const FPolygonID PolygonID, const FMeshPolygonContour& Contour, const bool bDeleteOrphanedEdges, const bool bDeleteOrphanedVertexInstances, TArray<FEdgeID>& OrphanedEdgeIDs, TArray<FVertexInstanceID>& OrphanedVertexInstanceIDs )
+{
+	for( const FVertexInstanceID VertexInstanceID : Contour.VertexInstanceIDs )
+	{
+		// Disconnect the polygon from the vertex instance
+		FMeshVertexInstance& VertexInstance = VertexInstances[ VertexInstanceID.GetValue() ];
+		verify( VertexInstance.ConnectedPolygons.Remove( PolygonID ) == 1 );
+
+		// Vertex instances which are only part of this contour (i.e. not shared by other polygons) are optionally deleted
+		if( bDeleteOrphanedVertexInstances && VertexInstance.ConnectedPolygons.Num() == 0 )
+		{
+			OrphanedVertexInstanceIDs.Add( VertexInstanceID );
+		}
+	}
+
+	// Remove reference to polygon from edges
+	static TArray<FEdgeID> EdgeIDs;
+	GetPolygonContourEdges( Contour, EdgeIDs );
+
+	for( const FEdgeID EdgeID : EdgeIDs )
+	{
+		FMeshEdge& Edge = Edges[ EdgeID.GetValue() ];
+		verify( Edge.ConnectedPolygons.Remove( PolygonID ) == 1 );
+
+		// Identify orphaned edges
+		if( bDeleteOrphanedEdges && Edge.ConnectedPolygons.Num() == 0 )
+		{
+			OrphanedEdgeIDs.Add( EdgeID );
+		}
+	}
+}
+
+
+void UEditableMesh::DeletePolygons( const TArray<FPolygonID>& PolygonIDsToDelete, const bool bDeleteOrphanedEdges, const bool bDeleteOrphanedVertices, const bool bDeleteOrphanedVertexInstances, const bool bDeleteEmptyPolygonGroups )
+{
+	UE_LOG( LogEditableMesh, Verbose, TEXT( "DeletePolygons: %s" ), *LogHelpers::ArrayToString( PolygonIDsToDelete ) );
+
+	// Back everything up
+	{
+		FCreatePolygonsChangeInput RevertInput;
+		RevertInput.PolygonsToCreate.Reserve( PolygonIDsToDelete.Num() );
+
+		// NOTE: We iterate backwards, to restore edges in the opposite order that we deleted them
+		for( int32 PolygonNumber = PolygonIDsToDelete.Num() - 1; PolygonNumber >= 0; --PolygonNumber )
+		{
+			const FPolygonID PolygonID = PolygonIDsToDelete[ PolygonNumber ];
+			const FMeshPolygon& Polygon = Polygons[ PolygonID.GetValue() ];
+
+			RevertInput.PolygonsToCreate.Emplace();
+			FPolygonToCreate& PolygonToCreate = RevertInput.PolygonsToCreate.Last();
+
+			PolygonToCreate.PolygonGroupID = Polygon.PolygonGroupID;
+			PolygonToCreate.OriginalPolygonID = PolygonID;
+
+			BackupPolygonContour( Polygon.PerimeterContour, PolygonToCreate.PerimeterVertices );
+			
+			PolygonToCreate.PolygonHoles.Reserve( Polygon.HoleContours.Num() );
+			for( const FMeshPolygonContour& HoleContour : Polygon.HoleContours )
+			{
+				PolygonToCreate.PolygonHoles.Emplace();
+				BackupPolygonContour( HoleContour, PolygonToCreate.PolygonHoles.Last().HoleVertices );
+			}
+		}
+
+		AddUndo( MakeUnique<FCreatePolygonsChange>( MoveTemp( RevertInput ) ) );
+	}
+
+	for( UEditableMeshAdapter* Adapter : Adapters )
+	{
+		Adapter->OnDeletePolygons( this, PolygonIDsToDelete );
+	}
+
+	// Delete the polygons
+	{
+		static TArray<FEdgeID> OrphanedEdgeIDs;
+		OrphanedEdgeIDs.Reset();
+
+		static TArray<FVertexInstanceID> OrphanedVertexInstanceIDs;
+		OrphanedVertexInstanceIDs.Reset();
+
+		static TArray<FPolygonGroupID> EmptyPolygonGroupIDs;
+		EmptyPolygonGroupIDs.Reset();
+
+		for( const FPolygonID PolygonID : PolygonIDsToDelete )
+		{
+			const FMeshPolygon& Polygon = Polygons[ PolygonID.GetValue() ];
+
+			DeletePolygonContour( PolygonID, Polygon.PerimeterContour, bDeleteOrphanedEdges, bDeleteOrphanedVertexInstances, OrphanedEdgeIDs, OrphanedVertexInstanceIDs );
+
+			for( const FMeshPolygonContour& HoleContour : Polygon.HoleContours )
+			{
+				DeletePolygonContour( PolygonID, HoleContour, bDeleteOrphanedEdges, bDeleteOrphanedVertexInstances, OrphanedEdgeIDs, OrphanedVertexInstanceIDs );
+			}
+
+			// Disconnect the polygon from the polygon group
+			const FPolygonGroupID PolygonGroupID = Polygon.PolygonGroupID;
+			FMeshPolygonGroup& PolygonGroup = PolygonGroups[ PolygonGroupID.GetValue() ];
+			verify( PolygonGroup.Polygons.Remove( PolygonID ) == 1 );
+
+			if( bDeleteEmptyPolygonGroups && PolygonGroup.Polygons.Num() == 0 )
+			{
+				EmptyPolygonGroupIDs.Add( PolygonGroupID );
+			}
+
+			// Delete the polygon
+			Polygons.RemoveAt( PolygonID.GetValue() );
+		}
+
+		// Remove vertex instances which are exclusively used by this polygon.
+		// We do not want this to remove orphaned vertices; this will optionally happen below when removing edges.
+		if( OrphanedVertexInstanceIDs.Num() > 0 )
+		{
+			DeleteVertexInstances( OrphanedVertexInstanceIDs, false );
+		}
+
+		// Remove any edges which may have been orphaned. This may also optionally remove any orphaned vertices.
+		// We can do this here because we know any edges which were orphaned will have had only a single vertex instance at each vertex.
+		// Therefore the vertex will now have no instances further to deleting them above.
+		// Note: there is never a situation where there could be orphaned vertices but not orphaned edges.
+		if( OrphanedEdgeIDs.Num() > 0 )
+		{
+			DeleteEdges( OrphanedEdgeIDs, bDeleteOrphanedVertices );
+		}
+
+		// Remove any empty polygon groups which may have resulted
+		if( EmptyPolygonGroupIDs.Num() > 0 )
+		{
+			DeletePolygonGroups( EmptyPolygonGroupIDs );
+		}
+	}
+
+	// If any of these polygons are in the pending list for triangulation or computing a new tangent basis, remove them
+	for( const FPolygonID PolygonID : PolygonIDsToDelete )
+	{
+		PolygonsPendingNewTangentBasis.Remove( PolygonID );
+		PolygonsPendingTriangulation.Remove( PolygonID );
+	}
+
+	UE_LOG( LogEditableMesh, Verbose, TEXT( "* DeletePolygons returned" ) );
+}
+
+
+void UEditableMesh::CreatePolygonGroups( const TArray<FPolygonGroupToCreate>& PolygonGroupsToCreate, TArray<FPolygonGroupID>& OutNewPolygonGroupIDs )
+{
+	// @todo mesheditor: Material, bEnableCollision and bCastShadow should all be implemented as PolygonGroup attributes
+
+	OutNewPolygonGroupIDs.Reset( PolygonGroupsToCreate.Num() );
+
+	// Reserve elements
+	PolygonGroups.Reserve( PolygonGroups.Num() + PolygonGroupsToCreate.Num() );
+
+	for( const FPolygonGroupToCreate& PolygonGroupToCreate : PolygonGroupsToCreate )
+	{
+		// Allocate polygon group
+		FPolygonGroupID PolygonGroupID = PolygonGroupToCreate.OriginalPolygonGroupID;
+		if( PolygonGroupID != FPolygonGroupID::Invalid )
+		{
+			PolygonGroups.Insert( PolygonGroupID.GetValue(), FMeshPolygonGroup() );
+		}
+		else
+		{
+			PolygonGroupID = FPolygonGroupID( PolygonGroups.Add( FMeshPolygonGroup() ) );
+		}
+
+		FMeshPolygonGroup& PolygonGroup = PolygonGroups[ PolygonGroupID.GetValue() ];
+
+		PolygonGroup.Material = PolygonGroupToCreate.Material;
+		PolygonGroup.bEnableCollision = PolygonGroupToCreate.bEnableCollision;
+		PolygonGroup.bCastShadow = PolygonGroupToCreate.bCastShadow;
+
+		OutNewPolygonGroupIDs.Add( PolygonGroupID );
+	}
+
+	for( UEditableMeshAdapter* Adapter : Adapters )
+	{
+		Adapter->OnCreatePolygonGroups( this, OutNewPolygonGroupIDs );
+	}
+
+	// Back up
+	{
+		FDeletePolygonGroupsChangeInput RevertInput;
+		RevertInput.PolygonGroupIDs.Reserve( OutNewPolygonGroupIDs.Num() );
+		for( int32 Index = OutNewPolygonGroupIDs.Num() - 1; Index >= 0; --Index )
+		{
+			RevertInput.PolygonGroupIDs.Add( OutNewPolygonGroupIDs[ Index ] );
+		}
+
+		AddUndo( MakeUnique<FDeletePolygonGroupsChange>( MoveTemp( RevertInput ) ) );
+	}
+}
+
+
+void UEditableMesh::DeletePolygonGroups( const TArray<FPolygonGroupID>& PolygonGroupIDs )
+{
+	// Back everything up
+	{
+		FCreatePolygonGroupsChangeInput RevertInput;
+
+		for( int32 Index = PolygonGroupIDs.Num() - 1; Index >= 0; --Index )
+		{
+			const FPolygonGroupID PolygonGroupID = PolygonGroupIDs[ Index ];
+			const FMeshPolygonGroup& PolygonGroup = PolygonGroups[ PolygonGroupID.GetValue() ];
+
+			RevertInput.PolygonGroupsToCreate.Emplace();
+			FPolygonGroupToCreate& PolygonGroupToCreate = RevertInput.PolygonGroupsToCreate.Last();
+
+			PolygonGroupToCreate.OriginalPolygonGroupID = PolygonGroupID;
+			PolygonGroupToCreate.Material = PolygonGroup.Material;
+			PolygonGroupToCreate.bCastShadow = PolygonGroup.bCastShadow;
+			PolygonGroupToCreate.bEnableCollision = PolygonGroup.bEnableCollision;
+		}
+
+		AddUndo( MakeUnique<FCreatePolygonGroupsChange>( MoveTemp( RevertInput ) ) );
+	}
+
+	for( UEditableMeshAdapter* Adapter : Adapters )
+	{
+		Adapter->OnDeletePolygonGroups( this, PolygonGroupIDs );
+	}
+
+	// Delete the polygon groups
+	{
+		for( const FPolygonGroupID PolygonGroupID : PolygonGroupIDs )
+		{
+			const FMeshPolygonGroup& PolygonGroup = PolygonGroups[ PolygonGroupID.GetValue() ];
+			check( PolygonGroup.Polygons.Num() == 0 );
+
+			PolygonGroups.RemoveAt( PolygonGroupID.GetValue() );
+		}
+	}
+}
+
+
+void UEditableMesh::SetVerticesAttributes( const TArray<FAttributesForVertex>& AttributesForVertices )
+{
+	UE_LOG( LogEditableMesh, Verbose, TEXT( "SetVerticesAttributes: %s" ), *LogHelpers::ArrayToString( AttributesForVertices ) );
+
+	FSetVerticesAttributesChangeInput RevertInput;
+
+	RevertInput.AttributesForVertices.Reserve( AttributesForVertices.Num() );
+	for( const FAttributesForVertex& AttributesForVertex : AttributesForVertices )
+	{
+		FAttributesForVertex& RevertVertex = *new( RevertInput.AttributesForVertices ) FAttributesForVertex();
+		RevertVertex.VertexID = AttributesForVertex.VertexID;
+
+		// Set the vertex attributes
+		{
+			RevertVertex.VertexAttributes.Attributes.Reserve( AttributesForVertex.VertexAttributes.Attributes.Num() );
+
+			for( const FMeshElementAttributeData& VertexAttribute : AttributesForVertex.VertexAttributes.Attributes )
+			{
+				// Back up the attribute
+				FMeshElementAttributeData& RevertAttribute = *new( RevertVertex.VertexAttributes.Attributes ) FMeshElementAttributeData(
+					VertexAttribute.AttributeName,
+					VertexAttribute.AttributeIndex,
+					GetVertexAttribute( AttributesForVertex.VertexID, VertexAttribute.AttributeName, VertexAttribute.AttributeIndex ) );
+
+				// Set the new attribute
+				SetVertexAttribute( AttributesForVertex.VertexID, VertexAttribute.AttributeName, VertexAttribute.AttributeIndex, VertexAttribute.AttributeValue );
+			}
+		}
+	}
+
+	AddUndo( MakeUnique<FSetVerticesAttributesChange>( MoveTemp( RevertInput ) ) );
+
+	UE_LOG( LogEditableMesh, Verbose, TEXT( "* SetVerticesAttributes returned" ) );
+}
+
+
+void UEditableMesh::SetVertexInstancesAttributes( const TArray<FAttributesForVertexInstance>& AttributesForVertexInstances )
+{
+	UE_LOG( LogEditableMesh, Verbose, TEXT( "SetVertexInstancesAttributes: %s" ), *LogHelpers::ArrayToString( AttributesForVertexInstances ) );
+
+	FSetVertexInstancesAttributesChangeInput RevertInput;
+
+	RevertInput.AttributesForVertexInstances.Reserve( AttributesForVertexInstances.Num() );
+	for( const FAttributesForVertexInstance& AttributesForVertexInstance : AttributesForVertexInstances )
+	{
+		FAttributesForVertexInstance& RevertVertexInstance = *new( RevertInput.AttributesForVertexInstances ) FAttributesForVertexInstance();
+		RevertVertexInstance.VertexInstanceID = AttributesForVertexInstance.VertexInstanceID;
+
+		// Set the vertex instance attributes
+		{
+			RevertVertexInstance.VertexInstanceAttributes.Attributes.Reserve( AttributesForVertexInstance.VertexInstanceAttributes.Attributes.Num() );
+
+			for( const FMeshElementAttributeData& VertexInstanceAttribute : AttributesForVertexInstance.VertexInstanceAttributes.Attributes )
+			{
+				// Back up the attribute
+				FMeshElementAttributeData& RevertAttribute = *new( RevertVertexInstance.VertexInstanceAttributes.Attributes ) FMeshElementAttributeData(
+					VertexInstanceAttribute.AttributeName,
+					VertexInstanceAttribute.AttributeIndex,
+					GetVertexInstanceAttribute( AttributesForVertexInstance.VertexInstanceID, VertexInstanceAttribute.AttributeName, VertexInstanceAttribute.AttributeIndex ) );
+
+				// Set the new attribute
+				SetVertexInstanceAttribute( AttributesForVertexInstance.VertexInstanceID, VertexInstanceAttribute.AttributeName, VertexInstanceAttribute.AttributeIndex, VertexInstanceAttribute.AttributeValue );
+			}
+		}
+	}
+
+	AddUndo( MakeUnique<FSetVertexInstancesAttributesChange>( MoveTemp( RevertInput ) ) );
+
+	UE_LOG( LogEditableMesh, Verbose, TEXT( "* SetVertexInstancesAttributes returned" ) );
+}
+
+
+void UEditableMesh::SetEdgesAttributes( const TArray<FAttributesForEdge>& AttributesForEdges )
+{
+	UE_LOG( LogEditableMesh, Verbose, TEXT( "SetEdgesAttributes: %s" ), *LogHelpers::ArrayToString( AttributesForEdges ) );
+
+	FSetEdgesAttributesChangeInput RevertInput;
+
+	RevertInput.AttributesForEdges.Reserve( AttributesForEdges.Num() );
+	for( const FAttributesForEdge& AttributesForEdge : AttributesForEdges )
+	{
+		FAttributesForEdge& RevertEdge = *new( RevertInput.AttributesForEdges ) FAttributesForEdge();
+		RevertEdge.EdgeID = AttributesForEdge.EdgeID;
+		RevertEdge.bIsUndo = true;
+
+		// Set the edge attributes
+		{
+			RevertEdge.EdgeAttributes.Attributes.Reserve( AttributesForEdge.EdgeAttributes.Attributes.Num() );
+
+			for( const FMeshElementAttributeData& EdgeAttribute : AttributesForEdge.EdgeAttributes.Attributes )
+			{
+				// Back up the attribute
+				FMeshElementAttributeData& RevertAttribute = *new( RevertEdge.EdgeAttributes.Attributes ) FMeshElementAttributeData(
+					EdgeAttribute.AttributeName,
+					EdgeAttribute.AttributeIndex,
+					GetEdgeAttribute( AttributesForEdge.EdgeID, EdgeAttribute.AttributeName, EdgeAttribute.AttributeIndex ) );
+
+				// Set the new attribute
+				SetEdgeAttribute( AttributesForEdge.EdgeID, EdgeAttribute.AttributeName, EdgeAttribute.AttributeIndex, EdgeAttribute.AttributeValue, AttributesForEdge.bIsUndo );
+			}
+		}
+	}
+
+	AddUndo( MakeUnique<FSetEdgesAttributesChange>( MoveTemp( RevertInput ) ) );
+
+	UE_LOG( LogEditableMesh, Verbose, TEXT( "* SetEdgesAttributes returned" ) );
+}
+
+
+const TArray<FName>& UEditableMesh::GetMergeableVertexInstanceAttributes()
+{
+	static TArray<FName> MergeableVertexInstanceAttributes;
+	if( MergeableVertexInstanceAttributes.Num() == 0 )
+	{
+		MergeableVertexInstanceAttributes.Add( UEditableMeshAttribute::VertexTextureCoordinate() );
+		MergeableVertexInstanceAttributes.Add( UEditableMeshAttribute::VertexColor() );
+	}
+	return MergeableVertexInstanceAttributes;
+}
+
+
+void UEditableMesh::GetMergedNamedVertexInstanceAttributeData( const TArray<FName>& AttributeNames, const FVertexInstanceID VertexInstanceID, const TArray<FMeshElementAttributeData>& AttributeValuesToMerge, TArray<FMeshElementAttributeData>& OutAttributes ) const
+{
+	OutAttributes.Reset();
+
+	for( const FName AttributeName : AttributeNames )
+	{
+		const int32 MaxAttributeIndex = GetMaxAttributeIndex( AttributeName );
+		for( int32 AttributeIndex = 0; AttributeIndex < MaxAttributeIndex; ++AttributeIndex )
+		{
+			const int32 ReplaceIndex =  AttributeValuesToMerge.IndexOfByPredicate(
+				[ AttributeName, AttributeIndex ]( const FMeshElementAttributeData& AttributeData )
+				{
+					return AttributeData.AttributeName == AttributeName && AttributeData.AttributeIndex == AttributeIndex;
+				}
+			);
+
+			if( ReplaceIndex == INDEX_NONE )
+			{
+				// If the named attribute and index doesn't appear in the list to merge, add the current attribute value to the list
+				OutAttributes.Emplace( AttributeName, AttributeIndex, GetVertexInstanceAttribute( VertexInstanceID, AttributeName, AttributeIndex ) );
+			}
+			else
+			{
+				// If the named attribute and index does appear in the list to merge, add that to the list instead
+				OutAttributes.Emplace( AttributeValuesToMerge[ ReplaceIndex ] );
+			}
+		}
+	}
+}
+
+
+void UEditableMesh::ChangePolygonsVertexInstances( const TArray<FChangeVertexInstancesForPolygon>& VertexInstancesForPolygons )
+{
+	UE_LOG( LogEditableMesh, Verbose, TEXT( "ChangePolygonsVertexInstances: %s" ), *LogHelpers::ArrayToString( VertexInstancesForPolygons ) );
+
+	// Back everything up
+	{
+		FChangePolygonsVertexInstancesChangeInput RevertInput;
+		RevertInput.VertexInstancesForPolygons.Reserve( VertexInstancesForPolygons.Num() );
+
+		// NOTE: We iterate backwards, to restore edges in the opposite order that we changed them
+		for( int32 Index = VertexInstancesForPolygons.Num() - 1; Index >= 0; --Index )
+		{
+			const FChangeVertexInstancesForPolygon& VertexInstancesForPolygon = VertexInstancesForPolygons[ Index ];
+			const FPolygonID PolygonID = VertexInstancesForPolygon.PolygonID;
+			const FMeshPolygon& Polygon = Polygons[ PolygonID.GetValue() ];
+
+			RevertInput.VertexInstancesForPolygons.Emplace();
+			FChangeVertexInstancesForPolygon& RevertVertexInstancesForPolygon = RevertInput.VertexInstancesForPolygons.Last();
+			RevertVertexInstancesForPolygon.PolygonID = PolygonID;
+
+			for( const FVertexIndexAndInstanceID& IndexAndInstance : VertexInstancesForPolygon.PerimeterVertexIndicesAndInstanceIDs )
+			{
+				RevertVertexInstancesForPolygon.PerimeterVertexIndicesAndInstanceIDs.Emplace();
+				FVertexIndexAndInstanceID& RevertIndexAndInstance = RevertVertexInstancesForPolygon.PerimeterVertexIndicesAndInstanceIDs.Last();
+				RevertIndexAndInstance.ContourIndex = IndexAndInstance.ContourIndex;
+				RevertIndexAndInstance.VertexInstanceID = Polygon.PerimeterContour.VertexInstanceIDs[ IndexAndInstance.ContourIndex ];
+			}
+
+			int32 HoleIndex = 0;
+			for( const FVertexInstancesForPolygonHole& VertexInstancesForPolygonHole : VertexInstancesForPolygon.VertexIndicesAndInstanceIDsForEachHole )
+			{
+				RevertVertexInstancesForPolygon.VertexIndicesAndInstanceIDsForEachHole.Emplace();
+				FVertexInstancesForPolygonHole& RevertVertexInstancesForPolygonHole = RevertVertexInstancesForPolygon.VertexIndicesAndInstanceIDsForEachHole.Last();
+
+				for( const FVertexIndexAndInstanceID& IndexAndInstance : VertexInstancesForPolygonHole.VertexIndicesAndInstanceIDs )
+				{
+					RevertVertexInstancesForPolygonHole.VertexIndicesAndInstanceIDs.Emplace();
+					FVertexIndexAndInstanceID& RevertIndexAndInstance = RevertVertexInstancesForPolygonHole.VertexIndicesAndInstanceIDs.Last();
+					RevertIndexAndInstance.ContourIndex = IndexAndInstance.ContourIndex;
+					RevertIndexAndInstance.VertexInstanceID = Polygon.HoleContours[ HoleIndex ].VertexInstanceIDs[ IndexAndInstance.ContourIndex ];
+				}
+
+				HoleIndex++;
+			}
+		}
+
+		AddUndo( MakeUnique<FChangePolygonsVertexInstancesChange>( MoveTemp( RevertInput ) ) );
+	}
+
+	static TArray<FPolygonID> PolygonIDs;
+	PolygonIDs.Reset( VertexInstancesForPolygons.Num() );
+
+	// Perform action
+	{
+		for( const FChangeVertexInstancesForPolygon& VertexInstancesForPolygon : VertexInstancesForPolygons )
+		{
+			const FPolygonID PolygonID = VertexInstancesForPolygon.PolygonID;
+			PolygonIDs.Add( PolygonID );
+			FMeshPolygon& Polygon = Polygons[ PolygonID.GetValue() ];
+
+			for( const FVertexIndexAndInstanceID& IndexAndInstance : VertexInstancesForPolygon.PerimeterVertexIndicesAndInstanceIDs )
+			{
+				// Disconnect old vertex instance from polygon, and connect new one
+				const FVertexInstanceID OldVertexInstanceID = Polygon.PerimeterContour.VertexInstanceIDs[ IndexAndInstance.ContourIndex ];
+				FMeshVertexInstance& OldVertexInstance = VertexInstances[ OldVertexInstanceID.GetValue() ];
+				verify( OldVertexInstance.ConnectedPolygons.Remove( PolygonID ) == 1 );
+
+				FMeshVertexInstance& NewVertexInstance = VertexInstances[ IndexAndInstance.VertexInstanceID.GetValue() ];
+				check( !NewVertexInstance.ConnectedPolygons.Contains( PolygonID ) );
+				NewVertexInstance.ConnectedPolygons.Add( PolygonID );
+
+				Polygon.PerimeterContour.VertexInstanceIDs[ IndexAndInstance.ContourIndex ] = IndexAndInstance.VertexInstanceID;
+
+				// Fix up triangle list
+				for( FMeshTriangle& Triangle : Polygon.Triangles )
+				{
+					for( int32 VertexIndex = 0; VertexIndex < 3; ++VertexIndex )
+					{
+						if( Triangle.GetVertexInstanceID( VertexIndex ) == OldVertexInstanceID )
+						{
+							Triangle.SetVertexInstanceID( VertexIndex, IndexAndInstance.VertexInstanceID );
+						}
+					}
+				}
+			}
+
+			int32 HoleIndex = 0;
+			for( const FVertexInstancesForPolygonHole& VertexInstancesForPolygonHole : VertexInstancesForPolygon.VertexIndicesAndInstanceIDsForEachHole )
+			{
+				for( const FVertexIndexAndInstanceID& IndexAndInstance : VertexInstancesForPolygonHole.VertexIndicesAndInstanceIDs )
+				{
+					// Disconnect old vertex instance from polygon, and connect new one
+					const FVertexInstanceID OldVertexInstanceID = Polygon.HoleContours[ HoleIndex ].VertexInstanceIDs[ IndexAndInstance.ContourIndex ];
+					FMeshVertexInstance& OldVertexInstance = VertexInstances[ OldVertexInstanceID.GetValue() ];
+					verify( OldVertexInstance.ConnectedPolygons.Remove( PolygonID ) == 1 );
+
+					FMeshVertexInstance& NewVertexInstance = VertexInstances[ IndexAndInstance.VertexInstanceID.GetValue() ];
+					check( !NewVertexInstance.ConnectedPolygons.Contains( PolygonID ) );
+					NewVertexInstance.ConnectedPolygons.Add( PolygonID );
+
+					Polygon.HoleContours[ HoleIndex ].VertexInstanceIDs[ IndexAndInstance.ContourIndex ] = IndexAndInstance.VertexInstanceID;
+
+					// Fix up triangle list
+					for( FMeshTriangle& Triangle : Polygon.Triangles )
+					{
+						for( int32 VertexIndex = 0; VertexIndex < 3; ++VertexIndex )
+						{
+							if( Triangle.GetVertexInstanceID( VertexIndex ) == OldVertexInstanceID )
+							{
+								Triangle.SetVertexInstanceID( VertexIndex, IndexAndInstance.VertexInstanceID );
+							}
+						}
+					}
+				}
+
+				HoleIndex++;
+			}
+		}
+	}
+
+	// Let the adapter deal with it
+	for( UEditableMeshAdapter* Adapter : Adapters )
+	{
+		Adapter->OnChangePolygonVertexInstances( this, PolygonIDs );
+	}
+
+	UE_LOG( LogEditableMesh, Verbose, TEXT( "* ChangePolygonsVertexInstances returned" ) );
+}
+
+
+FVertexInstanceID UEditableMesh::GetVertexInstanceInPolygonForVertex( const FPolygonID PolygonID, const FVertexID VertexID ) const
+{
+	const FMeshVertex& Vertex = Vertices[ VertexID.GetValue() ];
+
+	for( const FVertexInstanceID VertexInstanceID : Vertex.VertexInstanceIDs )
+	{
+		const FMeshVertexInstance& VertexInstance = VertexInstances[ VertexInstanceID.GetValue() ];
+		if( VertexInstance.ConnectedPolygons.Contains( PolygonID ) )
+		{
+			return VertexInstanceID;
+		}
+	}
+
+	return FVertexInstanceID::Invalid;
+}
+
+
+void UEditableMesh::GetConnectedSoftEdges( const FVertexID VertexID, TArray<FEdgeID>& OutConnectedSoftEdges ) const
+{
+	OutConnectedSoftEdges.Reset();
+
+	const FMeshVertex& Vertex = Vertices[ VertexID.GetValue() ];
+	for( const FEdgeID ConnectedEdgeID : Vertex.ConnectedEdgeIDs )
+	{
+		const FMeshEdge& Edge = Edges[ ConnectedEdgeID.GetValue() ];
+		if( !Edge.bIsHardEdge )
+		{
+			OutConnectedSoftEdges.Add( ConnectedEdgeID );
+		}
+	}
+}
+
+
+void UEditableMesh::GetVertexInstancesInSameSoftEdgedGroup( const FVertexID VertexID, const FPolygonID PolygonID, const bool bPolygonNotYetInitialized, TArray<FVertexInstanceID>& OutVertexInstanceIDs ) const
+{
+	OutVertexInstanceIDs.Reset();
+
+	// Cache a list of all soft edges which share this vertex.
+	// We're not interested in hard edges as they denote a transition to a different soft edged group.
+	// Note that a vertex is usually only treated as 'hard' if it is met by two hard edges: if a hard edge terminates at the vertex we're looking at,
+	// it may be possible to find an adjacency path from one side to the other, and hence does not represent a barrier between soft edged groups.
+	// The exception is if the mesh is not closed and the hard edge leads to the mesh boundary.
+	static TArray<FEdgeID> ConnectedSoftEdges;
+	GetConnectedSoftEdges( VertexID, ConnectedSoftEdges );
+
+	// Iterate through adjacent polygons which contain the given vertex, but don't cross a hard edge.
+	// Maintain a list of polygon IDs to be examined. Adjacents are added to the list if suitable.
+	// Add the start poly here.
+	static TArray<FPolygonID> PolygonsToCheck;
+	PolygonsToCheck.Reset();
+	PolygonsToCheck.Add( PolygonID );
+
+	int32 Index = 0;
+	while( Index < PolygonsToCheck.Num() )
+	{
+		const FPolygonID PolygonToCheck = PolygonsToCheck[ Index ];
+		Index++;
+
+		const FVertexInstanceID VertexInstanceID = GetVertexInstanceInPolygonForVertex( PolygonToCheck, VertexID );
+		if( VertexInstanceID != FVertexInstanceID::Invalid || ( bPolygonNotYetInitialized && Index == 0 ) )
+		{
+			// If we got here, either:
+			//  a) the polygon contains the vertex instance referencing the vertex we're looking at, or
+			//  b) the polygon is the one we passed in, which may not yet have its vertex instances assigned, and we've specified bPolygonNotYetInitialized=true to allow this case
+
+			// If this polygon contains an instance of this vertex, add it to the result.
+			if( VertexInstanceID != FVertexInstanceID::Invalid )
+			{
+				OutVertexInstanceIDs.AddUnique( VertexInstanceID );
+			}
+
+			// Now look at its adjacent polygons. If they are joined by a soft edge which includes the vertex we're interested in, we want to consider them.
+			// We take a shortcut by doing this process in reverse: we already know all the soft edges we are interested in, so check if any of them
+			// have the current polygon as an adjacent.
+			for( const FEdgeID ConnectedSoftEdge : ConnectedSoftEdges )
+			{
+				const FMeshEdge& Edge = Edges[ ConnectedSoftEdge.GetValue() ];
+				if( Edge.ConnectedPolygons.Contains( PolygonToCheck ) )
+				{
+					for( const FPolygonID AdjacentPolygon : Edge.ConnectedPolygons )
+					{
+						// Only add new polygons which haven't yet been added to the list. This prevents circular runs of polygons triggering infinite loops.
+						PolygonsToCheck.AddUnique( AdjacentPolygon );
+					}
+				}
+			}
+		}
+	}
+}
+
+
+void UEditableMesh::SetPolygonsVertexAttributes( const TArray<FVertexAttributesForPolygon>& VertexAttributesForPolygons )
+{
+	UE_LOG( LogEditableMesh, Verbose, TEXT( "SetPolygonsVertexAttributes: %s" ), *LogHelpers::ArrayToString( VertexAttributesForPolygons ) );
+
+	for( const FVertexAttributesForPolygon& VertexAttributesForPolygon : VertexAttributesForPolygons )
+	{
+		const FPolygonID PolygonID = VertexAttributesForPolygon.PolygonID;
+		FMeshPolygon& Polygon = Polygons[ PolygonID.GetValue() ];
+
+		SetPolygonContourVertexAttributes( Polygon.PerimeterContour, PolygonID, INDEX_NONE, VertexAttributesForPolygon.PerimeterVertexAttributeLists );
+
+		for( int32 HoleIndex = 0; HoleIndex < VertexAttributesForPolygon.VertexAttributeListsForEachHole.Num(); ++HoleIndex )
+		{
+			SetPolygonContourVertexAttributes( Polygon.HoleContours[ HoleIndex ], PolygonID, HoleIndex, VertexAttributesForPolygon.VertexAttributeListsForEachHole[ HoleIndex ].VertexAttributeList );
+		}
+	}
+
+	UE_LOG( LogEditableMesh, Verbose, TEXT( "* SetPolygonsVertexAttributes returned" ) );
+}
+
+
+bool UEditableMesh::AreAttributeListsEqual( const TArray<FMeshElementAttributeData>& A, const TArray<FMeshElementAttributeData>& B ) const
+{
+	// Check that two attribute lists have attributes of equal values.
+	// It is expected that each list contains exactly the same attribute names and indices, in the same order.
+	check( A.Num() == B.Num() );
+	for( int32 Index = 0; Index < A.Num(); ++Index )
+	{
+		check( A[ Index ].AttributeName == B[ Index ].AttributeName );
+		check( A[ Index ].AttributeIndex == B[ Index ].AttributeIndex );
+		if( A[ Index ].AttributeValue != B[ Index ].AttributeValue )
+		{
+			return false;
+		}
+	}
+	return true;
+};
+
+
+void UEditableMesh::SetPolygonContourVertexAttributes( FMeshPolygonContour& Contour, const FPolygonID PolygonID, const int32 HoleIndex, const TArray<FMeshElementAttributeList>& AttributeLists )
+{
+	const int32 NumContourVertices = Contour.VertexInstanceIDs.Num();
+	check( AttributeLists.Num() == NumContourVertices );
+
+	// Iterate round all polygons in the contour
+	int32 PrevIndex = NumContourVertices - 1;
+	for( int32 Index = 0; Index < NumContourVertices; ++Index )
+	{
+		const FMeshElementAttributeList& AttributeList = AttributeLists[ Index ];
+
+		// If there are no attributes to change, skip this index
+		if( AttributeList.Attributes.Num() == 0 )
+		{
+			continue;
+		}
+
+		// Get vertex instance and vertex.
+		const FVertexInstanceID VertexInstanceID = Contour.VertexInstanceIDs[ Index ];
+		const FMeshVertexInstance& VertexInstance = VertexInstances[ VertexInstanceID.GetValue() ];
+		const FVertexID VertexID = VertexInstance.VertexID;
+
+		// We need to determine whether changing the attributes on this polygon vertex would require vertex instances to be split or merged.
+		// First determine all the existing vertex instances which are in this soft edged group, that is to say,
+		// the group of polygons delimited by hard edges, of which this one is a member.
+		static TArray<FVertexInstanceID> PossibleVertexInstances;
+		PossibleVertexInstances.Reset();
+		const bool bPolygonNotYetInitialized = false;
+		GetVertexInstancesInSameSoftEdgedGroup( VertexID, PolygonID, bPolygonNotYetInitialized, PossibleVertexInstances );
+
+		// Now we have various possibilities:
+		//   1) If any of these vertex instances' attributes match our new polygon vertex attributes, just merge with that one, and disconnect the old one from the polygon.
+		//     1a) If that would cause the old vertex instance to become orphaned (because it is no longer connected to any polygon), it should be deleted.
+		// otherwise:
+		//   2) There are no vertex instances to merge into. 
+		//     2a) If the current vertex instance is not used by any other polygon, change its attributes in-place.
+		//   otherwise:
+		//     2b) Create a new vertex instance with the new attributes, and disconnect the old one from the polygon.
+
+		// Build the complete list of attributes for the polygon vertex being changed
+		static TArray<FMeshElementAttributeData> AllAttributes;
+		GetMergedNamedVertexInstanceAttributeData( GetMergeableVertexInstanceAttributes(), VertexInstanceID, AttributeList.Attributes, AllAttributes );
+
+		// First look for a vertex instance with matching attributes to ours
+		FVertexInstanceID MergedVertexInstanceID = FVertexInstanceID::Invalid;
+
+		for( const FVertexInstanceID PossibleVertexInstance : PossibleVertexInstances )
+		{
+			static TArray<FMeshElementAttributeData> AllAttributesToCompare;
+			GetMergedNamedVertexInstanceAttributeData( GetMergeableVertexInstanceAttributes(), PossibleVertexInstance, TArray<FMeshElementAttributeData>(), AllAttributesToCompare );
+
+			if( AreAttributeListsEqual( AllAttributes, AllAttributesToCompare ) )
+			{
+				// Found a vertex instance we can merge into
+				MergedVertexInstanceID = PossibleVertexInstance;
+				break;
+			}
+		}
+
+		if( MergedVertexInstanceID != FVertexInstanceID::Invalid )
+		{
+			// If we found a vertex instance we can merge into, change the perimeter vertex instance for the new one
+			static TArray<FChangeVertexInstancesForPolygon> VertexInstancesToChange;
+			VertexInstancesToChange.Reset();
+			VertexInstancesToChange.SetNum( 1 );
+
+			VertexInstancesToChange[ 0 ].PolygonID = PolygonID;
+			if( HoleIndex == INDEX_NONE )
+			{
+				VertexInstancesToChange[ 0 ].PerimeterVertexIndicesAndInstanceIDs.Emplace();
+				VertexInstancesToChange[ 0 ].PerimeterVertexIndicesAndInstanceIDs[ 0 ].ContourIndex = Index;
+				VertexInstancesToChange[ 0 ].PerimeterVertexIndicesAndInstanceIDs[ 0 ].VertexInstanceID = MergedVertexInstanceID;
+			}
+			else
+			{
+				VertexInstancesToChange[ 0 ].VertexIndicesAndInstanceIDsForEachHole.SetNum( HoleIndex + 1 );
+				VertexInstancesToChange[ 0 ].VertexIndicesAndInstanceIDsForEachHole[ HoleIndex ].VertexIndicesAndInstanceIDs.Emplace();
+				VertexInstancesToChange[ 0 ].VertexIndicesAndInstanceIDsForEachHole[ HoleIndex ].VertexIndicesAndInstanceIDs[ 0 ].ContourIndex = Index;
+				VertexInstancesToChange[ 0 ].VertexIndicesAndInstanceIDsForEachHole[ HoleIndex ].VertexIndicesAndInstanceIDs[ 0 ].VertexInstanceID = MergedVertexInstanceID;
+			}
+
+			ChangePolygonsVertexInstances( VertexInstancesToChange );
+
+			// If the change orphaned the old one, delete it
+			// @todo mesheditor: make this optional?
+			if( GetVertexInstanceConnectedPolygonCount( VertexInstanceID ) == 0 )
+			{
+				static TArray<FVertexInstanceID> VertexInstanceIDsToDelete;
+				VertexInstanceIDsToDelete.SetNumUninitialized( 1 );
+				VertexInstanceIDsToDelete[ 0 ] = VertexInstanceID;
+
+				const bool bDeleteOrphanedVertices = false;
+				DeleteVertexInstances( VertexInstanceIDsToDelete, bDeleteOrphanedVertices );
+			}
+		}
+		else
+		{
+			if( GetVertexInstanceConnectedPolygonCount( VertexInstanceID ) == 1 )
+			{
+				// This is the only polygon using this vertex instance, so change it in place
+				static TArray<FAttributesForVertexInstance> AttributesForVertexInstance;
+				AttributesForVertexInstance.Reset();
+				AttributesForVertexInstance.SetNum( 1 );
+				AttributesForVertexInstance[ 0 ].VertexInstanceID = VertexInstanceID;
+				AttributesForVertexInstance[ 0 ].VertexInstanceAttributes = AttributeList;
+
+				SetVertexInstancesAttributes( AttributesForVertexInstance );
+			}
+			else
+			{
+				// Split vertex instance: create a new one
+				static TArray<FVertexInstanceToCreate> VertexInstancesToCreate;
+				VertexInstancesToCreate.Reset();
+				VertexInstancesToCreate.SetNum( 1 );
+				VertexInstancesToCreate[ 0 ].VertexID = VertexID;
+				VertexInstancesToCreate[ 0 ].VertexInstanceAttributes.Attributes = AllAttributes;
+
+				static TArray<FVertexInstanceID> NewVertexInstanceIDs;
+				CreateVertexInstances( VertexInstancesToCreate, NewVertexInstanceIDs );
+
+				// and set it on the contour
+				static TArray<FChangeVertexInstancesForPolygon> VertexInstancesToChange;
+				VertexInstancesToChange.Reset();
+				VertexInstancesToChange.SetNum( 1 );
+
+				VertexInstancesToChange[ 0 ].PolygonID = PolygonID;
+				if( HoleIndex == INDEX_NONE )
+				{
+					VertexInstancesToChange[ 0 ].PerimeterVertexIndicesAndInstanceIDs.Emplace();
+					VertexInstancesToChange[ 0 ].PerimeterVertexIndicesAndInstanceIDs[ 0 ].ContourIndex = Index;
+					VertexInstancesToChange[ 0 ].PerimeterVertexIndicesAndInstanceIDs[ 0 ].VertexInstanceID = NewVertexInstanceIDs[ 0 ];
+				}
+				else
+				{
+					VertexInstancesToChange[ 0 ].VertexIndicesAndInstanceIDsForEachHole.SetNum( HoleIndex + 1 );
+					VertexInstancesToChange[ 0 ].VertexIndicesAndInstanceIDsForEachHole[ HoleIndex ].VertexIndicesAndInstanceIDs.Emplace();
+					VertexInstancesToChange[ 0 ].VertexIndicesAndInstanceIDsForEachHole[ HoleIndex ].VertexIndicesAndInstanceIDs[ 0 ].ContourIndex = Index;
+					VertexInstancesToChange[ 0 ].VertexIndicesAndInstanceIDsForEachHole[ HoleIndex ].VertexIndicesAndInstanceIDs[ 0 ].VertexInstanceID = NewVertexInstanceIDs[ 0 ];
+				}
+
+				ChangePolygonsVertexInstances( VertexInstancesToChange );
+			}
+		}
+
+		PrevIndex = Index;
+	}
+}
+
+
+void UEditableMesh::TryToRemovePolygonEdge( const FEdgeID EdgeID, bool& bOutWasEdgeRemoved, FPolygonID& OutNewPolygonID )
+{
+	UE_LOG( LogEditableMesh, Verbose, TEXT( "TryToRemovePolygonEdge: %s" ), *EdgeID.ToString() );
+
+	bOutWasEdgeRemoved = false;
+	OutNewPolygonID = FPolygonID::Invalid;
+
+	// If the edge is not shared by at least two polygons, we can't remove it.  (We would have to delete the polygon that owns
+	// this edge, which is not the intent of this feature.).  We also can't cleanly remove edges that are joining more than two 
+	// polygons.  We need to create a new polygon the two polygons, and if there were more than two then the remaining polygons 
+	// would be left disconnected after our edge is gone
+	const int32 ConnectedPolygonCount = GetEdgeConnectedPolygonCount( EdgeID );
+	if( ConnectedPolygonCount == 2 )
+	{
+		// Verify that both vertices on either end of this edge are connected to polygon (non-internal) edges.
+		// We currently do not expect to support internal triangles that don't touch the polygonal boundaries at all.
+		bool bBothVerticesConnectToPolygonEdges = true;
+		for( int32 EdgeVertexNumber = 0; EdgeVertexNumber < 2; ++EdgeVertexNumber )
+		{
+			const FVertexID VertexID = GetEdgeVertex( EdgeID, EdgeVertexNumber );
+
+			bool bVertexConnectsToPolygonEdge = false;
+
+			const int32 ConnectedEdgeCount = GetVertexConnectedEdgeCount( VertexID );
+			for( int32 ConnectedEdgeNumber = 0; ConnectedEdgeNumber < ConnectedEdgeCount; ++ConnectedEdgeNumber )
+			{
+				const FEdgeID OtherEdgeID = GetVertexConnectedEdge( VertexID, ConnectedEdgeNumber );
+				if( OtherEdgeID != EdgeID )
+				{
+					bVertexConnectsToPolygonEdge = true;
+					break;
+				}
+			}
+
+			if( !bVertexConnectsToPolygonEdge )
+			{
+				bBothVerticesConnectToPolygonEdges = false;
+			}
+		}
+
+		if( bBothVerticesConnectToPolygonEdges )
+		{
+			TArray<FVertexID> PolygonAVertexIDs, PolygonBVertexIDs;
+
+			const FPolygonID PolygonAID = GetEdgeConnectedPolygon( EdgeID, 0 );
+			GetPolygonPerimeterVertices( PolygonAID, /* Out */ PolygonAVertexIDs );
+
+			const FPolygonID PolygonBID = GetEdgeConnectedPolygon( EdgeID, 1 );
+			GetPolygonPerimeterVertices( PolygonBID, /* Out */ PolygonBVertexIDs );
+
+			// If the polygons are in different polygon groups, we can't remove the edge because we can't determine
+			// which polygon group the replacing polygon should belong to
+			// @todo mesheditor: We could just allow this, and use the first polygon group ID.  It's just a bit weird in the UI.
+			const FPolygonGroupID PolygonGroupID = GetGroupForPolygon( PolygonAID );
+			if( PolygonGroupID == GetGroupForPolygon( PolygonBID ) )
+			{
+				// Create a polygon by combining the edges from either polygon we're connected to, omitting the edge we're removing
+				TArray<FVertexAndAttributes> NewPolygonVertices;
+				{
+					const FVertexID EdgeVertexIDA = GetEdgeVertex( EdgeID, 0 );
+					const FVertexID EdgeVertexIDB = GetEdgeVertex( EdgeID, 1 );
+
+					// @todo mesheditor urgent: This code was written with the assumption that the total vertex count would
+					// end up being two less than the total vertex count of the two polygons we're replacing.  However, this
+					// might not be true for cases where either polygon shares more than one edge with the other, right?
+
+					// Find the edge vertices in the first polygon
+					int32 EdgeStartsAtVertexInPolygonA = INDEX_NONE;
+					for( int32 PolygonAVertexNumber = 0; PolygonAVertexNumber < PolygonAVertexIDs.Num(); ++PolygonAVertexNumber )
+					{
+						const FVertexID PolygonAVertexID = PolygonAVertexIDs[ PolygonAVertexNumber ];
+						const FVertexID PolygonANextVertexID = PolygonAVertexIDs[ ( PolygonAVertexNumber + 1 ) % PolygonAVertexIDs.Num() ];
+
+						if( ( PolygonAVertexID == EdgeVertexIDA || PolygonAVertexID == EdgeVertexIDB ) &&
+							( PolygonANextVertexID == EdgeVertexIDA || PolygonANextVertexID == EdgeVertexIDB ) )
+						{
+							EdgeStartsAtVertexInPolygonA = PolygonAVertexNumber;
+							break;
+						}
+					}
+					check( EdgeStartsAtVertexInPolygonA != INDEX_NONE );
+					const int32 EdgeEndsAtVertexInPolygonA = ( EdgeStartsAtVertexInPolygonA + 1 ) % PolygonAVertexIDs.Num();
+
+
+					// Find the edge vertices in the second polygon
+					int32 EdgeStartsAtVertexInPolygonB = INDEX_NONE;
+					for( int32 PolygonBVertexNumber = 0; PolygonBVertexNumber < PolygonBVertexIDs.Num(); ++PolygonBVertexNumber )
+					{
+						const FVertexID PolygonBVertexID = PolygonBVertexIDs[ PolygonBVertexNumber ];
+						const FVertexID PolygonBNextVertexID = PolygonBVertexIDs[ ( PolygonBVertexNumber + 1 ) % PolygonBVertexIDs.Num() ];
+
+						if( ( PolygonBVertexID == EdgeVertexIDA || PolygonBVertexID == EdgeVertexIDB ) &&
+							( PolygonBNextVertexID == EdgeVertexIDA || PolygonBNextVertexID == EdgeVertexIDB ) )
+						{
+							EdgeStartsAtVertexInPolygonB = PolygonBVertexNumber;
+							break;
+						}
+					}
+					check( EdgeStartsAtVertexInPolygonB != INDEX_NONE );
+					const int32 EdgeEndsAtVertexInPolygonB = ( EdgeStartsAtVertexInPolygonB + 1 ) % PolygonBVertexIDs.Num();
+
+					
+					// Do the polygons wind in the same direction?  If they do, the edge order will be reversed.
+					const bool bPolygonsWindInSameDirection = PolygonAVertexIDs[ EdgeStartsAtVertexInPolygonA ] != PolygonBVertexIDs[ EdgeStartsAtVertexInPolygonB ];
+
+
+					// Start adding vertices from the first polygon, starting with the vertex right after the edge we're removing.
+					// We'll continue to add vertices from this polygon until we reach back around to that edge.
+					const int32 PolygonAStartVertex = EdgeEndsAtVertexInPolygonA;
+					const int32 PolygonAEndVertex = EdgeStartsAtVertexInPolygonA;
+					for( int32 PolygonAVertexNumber = PolygonAStartVertex; PolygonAVertexNumber != PolygonAEndVertex; )
+					{
+						const FVertexID PolygonAVertexID = PolygonAVertexIDs[ PolygonAVertexNumber ];
+
+						FVertexAndAttributes& NewPolygonVertex = *new( NewPolygonVertices ) FVertexAndAttributes();
+						NewPolygonVertex.VertexID = PolygonAVertexID;
+
+						// Store off vertex attributes so we can re-apply them later to the new polygon
+						// @todo mesheditor: utility method which automatically does this with all known polygon vertex attributes? (or redundant when using "uber mesh"?)
+						for( int32 TextureCoordinateIndex = 0; TextureCoordinateIndex < TextureCoordinateCount; ++TextureCoordinateIndex )
+						{
+							NewPolygonVertex.PolygonVertexAttributes.Attributes.Emplace(
+								UEditableMeshAttribute::VertexTextureCoordinate(),
+								TextureCoordinateIndex,
+								GetPolygonPerimeterVertexAttribute( PolygonAID, PolygonAVertexNumber, UEditableMeshAttribute::VertexTextureCoordinate(), TextureCoordinateIndex )
+							);
+						}
+
+						NewPolygonVertex.PolygonVertexAttributes.Attributes.Emplace(
+							UEditableMeshAttribute::VertexColor(),
+							0,
+							GetPolygonPerimeterVertexAttribute( PolygonAID, PolygonAVertexNumber, UEditableMeshAttribute::VertexColor(), 0 )
+						);
+
+						PolygonAVertexNumber = ( PolygonAVertexNumber + 1 ) % PolygonAVertexIDs.Num();
+					}
+
+					// Now add vertices from the second polygon
+					const int32 PolygonBStartVertex = bPolygonsWindInSameDirection ? EdgeEndsAtVertexInPolygonB : EdgeStartsAtVertexInPolygonB;
+					const int32 PolygonBEndVertex = bPolygonsWindInSameDirection ? EdgeStartsAtVertexInPolygonB : EdgeEndsAtVertexInPolygonB;
+					for( int32 PolygonBVertexNumber = PolygonBStartVertex; PolygonBVertexNumber != PolygonBEndVertex; )
+					{
+						const FVertexID PolygonBVertexID = PolygonBVertexIDs[ PolygonBVertexNumber ];
+
+						FVertexAndAttributes& NewPolygonVertex = *new( NewPolygonVertices ) FVertexAndAttributes();
+						NewPolygonVertex.VertexID = PolygonBVertexID;
+
+						// Store off vertex attributes so we can re-apply them later to the new polygon
+						for( int32 TextureCoordinateIndex = 0; TextureCoordinateIndex < TextureCoordinateCount; ++TextureCoordinateIndex )
+						{
+							NewPolygonVertex.PolygonVertexAttributes.Attributes.Emplace(
+								UEditableMeshAttribute::VertexTextureCoordinate(),
+								TextureCoordinateIndex,
+								GetPolygonPerimeterVertexAttribute( PolygonBID, PolygonBVertexNumber, UEditableMeshAttribute::VertexTextureCoordinate(), TextureCoordinateIndex )
+							);
+						}
+
+						NewPolygonVertex.PolygonVertexAttributes.Attributes.Emplace(
+							UEditableMeshAttribute::VertexColor(),
+							0,
+							GetPolygonPerimeterVertexAttribute( PolygonBID, PolygonBVertexNumber, UEditableMeshAttribute::VertexColor(), 0 )
+						);
+
+						if( bPolygonsWindInSameDirection )
+						{
+							// Iterate forwards
+							PolygonBVertexNumber = ( PolygonBVertexNumber + 1 ) % PolygonBVertexIDs.Num();
+						}
+						else
+						{
+							// Iterate backwards
+							if( PolygonBVertexNumber == 0 )	   // @todo mesheditor: Not tested yet.  Need to connect two faces that wind in opposite directions
+							{
+								PolygonBVertexNumber = PolygonBVertexIDs.Num() - 1;
+							}
+							else
+							{
+								--PolygonBVertexNumber;
+							}
+						}
+					}
+				}
+
+				// OK, we can go ahead and delete the edge and its connected polygons.  We do NOT want to delete any orphaned
+				// edges or vertices though.  We're going to create a new polygon that connects to those right afterwards.
+				// @todo mesheditor: simplify this action so we just create the new polygon from existing vertex instances, instead of remembering all the attributes.
+				const bool bDeleteOrphanedEdges = false;
+				const bool bDeleteOrphanedVertices = false;
+				const bool bDeleteOrphanedVertexInstances = false;
+				const bool bDeleteEmptyPolygonGroups = false;
+				DeleteEdgeAndConnectedPolygons( EdgeID, bDeleteOrphanedEdges, bDeleteOrphanedVertices, bDeleteOrphanedVertexInstances, bDeleteEmptyPolygonGroups );
+
+				// Now create a new polygon to replace the two polygons we deleted
+				{
+					// Create the polygon
+					static TArray<FPolygonToCreate> PolygonsToCreate;
+					PolygonsToCreate.Reset();
+
+					// @todo mesheditor perf: Ideally we support creating multiple polygons at once and batching up the work
+					FPolygonToCreate& PolygonToCreate = *new( PolygonsToCreate ) FPolygonToCreate();
+					PolygonToCreate.PolygonGroupID = PolygonGroupID;
+					PolygonToCreate.PerimeterVertices = NewPolygonVertices;	// @todo mesheditor perf: Extra allocatons/copies: Ideally MoveTemp() here but we can't move a STATIC local!
+					// @todo mesheditor hole: Hole support needed!
+
+					static TArray<FPolygonID> NewPolygonIDs;
+					NewPolygonIDs.Reset();
+					static TArray<FEdgeID> NewEdgeIDs;
+					NewEdgeIDs.Reset();
+					CreatePolygons( PolygonsToCreate, /* Out */ NewPolygonIDs, /* Out */ NewEdgeIDs );	// @todo mesheditor perf: Extra allocatons/copies: Ideally MoveTemp() here but we can't move a STATIC local!
+
+					OutNewPolygonID = NewPolygonIDs[ 0 ];
+				}
+
+				bOutWasEdgeRemoved = true;
+			}
+		}
+	}
+
+	UE_LOG( LogEditableMesh, Verbose, TEXT( "* TryToRemovePolygonEdge returned %s" ), *OutNewPolygonID.ToString() );
+}
+
+
+void UEditableMesh::TryToRemoveVertex( const FVertexID VertexID, bool& bOutWasVertexRemoved, FEdgeID& OutNewEdgeID )
+{
+	UE_LOG( LogEditableMesh, Verbose, TEXT( "TryToRemoveVertex: %s" ), *VertexID.ToString() );
+
+	bOutWasVertexRemoved = false;
+	OutNewEdgeID = FEdgeID::Invalid;
+
+	// We only support removing vertices that are shared by just two edges
+	const int32 ConnectedEdgeCount = GetVertexConnectedEdgeCount( VertexID );
+	if( ConnectedEdgeCount == 2 )
+	{
+		// Get the two vertices on the other edge of either edge
+		FVertexID NewEdgeVertexIDs[ 2 ];
+		for( int32 EdgeNumber = 0; EdgeNumber < ConnectedEdgeCount; ++EdgeNumber )
+		{
+			const FEdgeID OtherEdgeID = GetVertexConnectedEdge( VertexID, EdgeNumber );
+			FVertexID OtherEdgeVertexIDs[ 2 ];
+			GetEdgeVertices( OtherEdgeID, /* Out */ OtherEdgeVertexIDs[0], /* Out */ OtherEdgeVertexIDs[1] );
+
+			NewEdgeVertexIDs[ EdgeNumber ] = OtherEdgeVertexIDs[ 0 ] == VertexID ? OtherEdgeVertexIDs[ 1 ] : OtherEdgeVertexIDs[ 0 ];
+		}
+
+		// Try to preserve attributes of the edges we're deleting.  We'll take the attributes from the first
+		// edge and apply them to the newly created edge
+		static TArray<FMeshElementAttributeData> NewEdgeAttributes;
+		{
+			const FEdgeID OtherEdgeID = GetVertexConnectedEdge( VertexID, 0 );		// @todo mesheditor: This is sort of subjective how we handle this (taking attributes from the first edge)
+
+			NewEdgeAttributes.Reset();
+			for( const FName AttributeName : UEditableMesh::GetValidEdgeAttributes() )
+			{
+				const int32 MaxAttributeIndex = GetMaxAttributeIndex( AttributeName );
+				for( int32 AttributeIndex = 0; AttributeIndex < MaxAttributeIndex; ++AttributeIndex )
+				{
+					FMeshElementAttributeData& EdgeAttribute = *new( NewEdgeAttributes ) FMeshElementAttributeData(
+						AttributeName,
+						AttributeIndex,
+						GetEdgeAttribute( OtherEdgeID, AttributeName, AttributeIndex ) );
+				}
+			}
+		}
+
+		// The new edge will be connected to the same polygons as both of the edges we're replacing.  Because
+		// we only support deleting a vertex shared by two edges, the two edges are guaranteed to be connected
+		// to the same exact polygons
+		static TArray<FPolygonID> NewEdgeConnectedPolygons;
+		GetVertexConnectedPolygons( VertexID, /* Out */ NewEdgeConnectedPolygons );
+
+		// Remove the vertex from its connected polygons
+		{
+			for( const FPolygonID PolygonID : NewEdgeConnectedPolygons )
+			{
+				// @todo mesheditor holes
+				const int32 PolygonVertexNumber = FindPolygonPerimeterVertexNumberForVertex( PolygonID, VertexID );
+				check( PolygonVertexNumber != INDEX_NONE );
+				const bool bDeleteOrphanedVertexInstances = false;
+				RemovePolygonPerimeterVertices( PolygonID, PolygonVertexNumber, 1, bDeleteOrphanedVertexInstances );
+			}
+		}
+
+		// Delete the two edges
+		{
+			static TArray<FEdgeID> EdgeIDsToDelete;
+			EdgeIDsToDelete.Reset();
+			for( int32 EdgeNumber = 0; EdgeNumber < ConnectedEdgeCount; ++EdgeNumber )
+			{
+				EdgeIDsToDelete.Add( GetVertexConnectedEdge( VertexID, EdgeNumber ) );
+			}
+
+			// NOTE: We can't delete the orphan vertex yet because the polygon triangles are still referencing
+			// its rendering vertices.  We'll delete the edges, retriangulate, then delete the vertex afterwards.
+			const bool bDeleteOrphanedVertices = false;
+			DeleteEdges( EdgeIDsToDelete, bDeleteOrphanedVertices );
+		}
+
+		FEdgeID NewEdgeID = FEdgeID::Invalid;
+
+		// Create a new edge to replace the vertex and two edges we deleted
+		{
+			static TArray<FEdgeToCreate> EdgesToCreate;
+			EdgesToCreate.Reset();
+
+			FEdgeToCreate& EdgeToCreate = *new( EdgesToCreate ) FEdgeToCreate;
+			EdgeToCreate.VertexID0 = NewEdgeVertexIDs[ 0 ];
+			EdgeToCreate.VertexID1 = NewEdgeVertexIDs[ 1 ];
+			EdgeToCreate.ConnectedPolygons = NewEdgeConnectedPolygons;
+			EdgeToCreate.EdgeAttributes.Attributes = NewEdgeAttributes;
+
+			static TArray<FEdgeID> NewEdgeIDs;
+			NewEdgeIDs.Reset();
+			CreateEdges( EdgesToCreate, /* Out */ NewEdgeIDs );
+
+			NewEdgeID = NewEdgeIDs[ 0 ];
+		}
+
+		// Update the normals of the affected polygons
+		PolygonsPendingNewTangentBasis.Append( NewEdgeConnectedPolygons );
+
+		// Retriangulate all of the affected polygons
+		PolygonsPendingTriangulation.Append( NewEdgeConnectedPolygons );
+
+		// Delete the vertex instances and subsequently orphaned vertex
+		{
+			FMeshVertex& Vertex = Vertices[ VertexID.GetValue() ];
+			const bool bDeleteOrphanedVertices = true;
+			DeleteVertexInstances( Vertex.VertexInstanceIDs, bDeleteOrphanedVertices );
+		}
+
+		bOutWasVertexRemoved = true;
+		OutNewEdgeID = NewEdgeID;
+	}
+
+	UE_LOG( LogEditableMesh, Verbose, TEXT( "* TryToRemoveVertex returned %s" ), *OutNewEdgeID.ToString() );
+}
+
+
+void UEditableMesh::ExtrudePolygons( const TArray<FPolygonID>& PolygonIDs, const float ExtrudeDistance, const bool bKeepNeighborsTogether, TArray<FPolygonID>& OutNewExtrudedFrontPolygons )
+{
+	UE_LOG( LogEditableMesh, Verbose, TEXT( "ExtrudePolygons: %s" ), *LogHelpers::ArrayToString( PolygonIDs ) );
+
+	// @todo mesheditor perf: We can make this much faster by batching up polygons together.  Just be careful about how neighbors are handled.
+	OutNewExtrudedFrontPolygons.Reset();
+
+	// Convert our incoming polygon array to a TSet so we can lookup quickly to see which polygons in the mesh sare members of the set
+	static TSet<FPolygonID> PolygonsSet;
+	PolygonsSet.Reset();
+	PolygonsSet.Append( PolygonIDs );
+	
+	static TArray<FPolygonID> AllNewPolygons;
+	AllNewPolygons.Reset();
+
+	static TArray<FAttributesForEdge> AttributesForEdges;
+	AttributesForEdges.Reset();
+
+	static TArray<FAttributesForVertex> AttributesForVertices;
+	AttributesForVertices.Reset();
+
+	static TArray<FVertexAttributesForPolygon> VertexAttributesForPolygons;
+	VertexAttributesForPolygons.Reset();
+
+	// First, let's figure out which of the polygons we were asked to extrude share edges or vertices.  We'll keep those
+	// edges intact!
+	static TMap<FEdgeID, uint32> EdgeUsageCounts;	// Maps an edge ID to the number of times it is referenced by the incoming polygons
+	EdgeUsageCounts.Reset();
+	static TSet<FVertexID> UniqueVertexIDs;
+	UniqueVertexIDs.Reset();
+
+	for( int32 PolygonIter = 0; PolygonIter < PolygonIDs.Num(); ++PolygonIter )
+	{
+		const FPolygonID PolygonID = PolygonIDs[ PolygonIter ];
+
+		static TArray<FEdgeID> PolygonPerimeterEdgeIDs;
+		GetPolygonPerimeterEdges( PolygonID, /* Out */ PolygonPerimeterEdgeIDs );
+
+		for( const FEdgeID EdgeID : PolygonPerimeterEdgeIDs )
+		{
+			FVertexID EdgeVertexIDs[ 2 ];
+			GetEdgeVertices( EdgeID, /* Out */ EdgeVertexIDs[ 0 ], /* Out */ EdgeVertexIDs[ 1 ] );
+
+			uint32* EdgeUsageCountPtr = EdgeUsageCounts.Find( EdgeID );
+			if( EdgeUsageCountPtr == nullptr )
+			{
+				EdgeUsageCounts.Add( EdgeID, 1 );
+			}
+			else
+			{
+				++( *EdgeUsageCountPtr );
+			}
+		}
+
+
+		static TArray<FVertexID> PolygonPerimeterVertexIDs;
+		this->GetPolygonPerimeterVertices( PolygonID, /* Out */ PolygonPerimeterVertexIDs );
+
+		for( const FVertexID VertexID : PolygonPerimeterVertexIDs )
+		{
+			UniqueVertexIDs.Add( VertexID );
+		}
+	}
+
+	const int32 NumVerticesToCreate = UniqueVertexIDs.Num();
+
+	// Create new vertices for all of the extruded polygons
+	static TArray<FVertexID> ExtrudedVertexIDs;
+	ExtrudedVertexIDs.Reset();
+	ExtrudedVertexIDs.Reserve( NumVerticesToCreate );
+	CreateEmptyVertexRange( NumVerticesToCreate, /* Out */ ExtrudedVertexIDs );
+	int32 NextAvailableExtrudedVertexIDNumber = 0;
+
+
+	static TMap<FVertexID, FVertexID> VertexIDToExtrudedCopy;
+	VertexIDToExtrudedCopy.Reset();
+
+	for( int32 PassIndex = 0; PassIndex < 2; ++PassIndex )
+	{
+		// Extrude all of the shared edges first, then do the non-shared edges.  This is to make sure that a vertex doesn't get offset
+		// without taking into account all of the connected polygons in our set
+		const bool bIsExtrudingSharedEdges = ( PassIndex == 0 );
+
+		for( int32 PolygonIter = 0; PolygonIter < PolygonIDs.Num(); ++PolygonIter )
+		{
+			const FPolygonID PolygonID = PolygonIDs[ PolygonIter ];
+			const FPolygonGroupID PolygonGroupID = GetGroupForPolygon( PolygonID );
+
+			if( !bKeepNeighborsTogether )
+			{
+				VertexIDToExtrudedCopy.Reset();
+			}
+
+			// Map all of the edge vertices to their new extruded counterpart
+			const int32 PerimeterEdgeCount = GetPolygonPerimeterEdgeCount( PolygonID );
+			for( int32 PerimeterEdgeNumber = 0; PerimeterEdgeNumber < PerimeterEdgeCount; ++PerimeterEdgeNumber )
+			{
+				// @todo mesheditor perf: We can change GetPolygonPerimeterEdges() to have a version that returns whether winding is reversed or not, and avoid this call entirely.
+				// @todo mesheditor perf: O(log N^2) iteration here. For every edge, for every edge up to this index.  Need to clean this up. 
+				//		--> Also, there are quite a few places where we are stepping through edges in perimeter-order.  We need to have a nice way to walk that.
+				bool bEdgeWindingIsReversedForPolygon;
+				const FEdgeID EdgeID = this->GetPolygonPerimeterEdge( PolygonID, PerimeterEdgeNumber, /* Out */ bEdgeWindingIsReversedForPolygon );
+
+				const bool bIsSharedEdge = bKeepNeighborsTogether && EdgeUsageCounts[ EdgeID ] > 1;
+				if( bIsSharedEdge == bIsExtrudingSharedEdges )
+				{
+					FVertexID EdgeVertexIDs[ 2 ];
+					GetEdgeVertices( EdgeID, /* Out */ EdgeVertexIDs[ 0 ], /* Out */ EdgeVertexIDs[ 1 ] );
+					if( bEdgeWindingIsReversedForPolygon )
+					{
+						::Swap( EdgeVertexIDs[ 0 ], EdgeVertexIDs[ 1 ] );
+					}
+
+					if( !bIsSharedEdge )
+					{
+						// After extruding, all of the edges of the original polygon become hard edges
+						FAttributesForEdge& AttributesForEdge = *new( AttributesForEdges ) FAttributesForEdge();
+						AttributesForEdge.EdgeID = EdgeID;
+						AttributesForEdge.EdgeAttributes.Attributes.Add( FMeshElementAttributeData( UEditableMeshAttribute::EdgeIsHard(), 0, FVector4( 1.0f ) ) );
+					}
+
+					FVertexID ExtrudedEdgeVertexIDs[ 2 ];
+					for( int32 EdgeVertexNumber = 0; EdgeVertexNumber < 2; ++EdgeVertexNumber )
+					{
+						const FVertexID EdgeVertexID = EdgeVertexIDs[ EdgeVertexNumber ];
+						FVertexID* ExtrudedEdgeVertexIDPtr = VertexIDToExtrudedCopy.Find( EdgeVertexID );
+
+						// @todo mesheditor extrude: Ideally we would detect whether the vertex that was already extruded came from a edge
+						// from a polygon that does not actually share an edge with any polygons this polygon shares an edge with.  This
+						// would avoid the problem where extruding two polygons that are connected only by a vertex are not extruded
+						// separately.
+						const bool bVertexIsSharedByAnEdgeOfAnotherSelectedPolygon = false;
+						if( ExtrudedEdgeVertexIDPtr != nullptr && !bVertexIsSharedByAnEdgeOfAnotherSelectedPolygon )
+						{
+							ExtrudedEdgeVertexIDs[ EdgeVertexNumber ] = *ExtrudedEdgeVertexIDPtr;
+						}
+						else
+						{
+							// Create a copy of this vertex for the extruded face
+							const FVertexID ExtrudedVertexID = ExtrudedVertexIDs[ NextAvailableExtrudedVertexIDNumber++ ];
+
+							ExtrudedEdgeVertexIDPtr = &VertexIDToExtrudedCopy.Add( EdgeVertexID, ExtrudedVertexID );
+							// Push the vertex out along the polygon's normal
+							const FVector OriginalVertexPosition = this->GetVertexAttribute( EdgeVertexID, UEditableMeshAttribute::VertexPosition(), 0 );
+
+							FVector ExtrudedVertexPosition;
+							if( bIsSharedEdge )
+							{
+								// Get all of the polygons that share this edge that were part of the set of polygons passed in.  We'll
+								// generate an extrude direction that's the average of those polygon normals.
+								FVector ExtrudeDirection = FVector::ZeroVector;
+
+								static TArray<FPolygonID> ConnectedPolygonIDs;
+								GetVertexConnectedPolygons( EdgeVertexID, /* Out */ ConnectedPolygonIDs );
+
+								static TArray<FPolygonID> NeighborPolygonIDs;
+								NeighborPolygonIDs.Reset();
+								for( const FPolygonID ConnectedPolygonID : ConnectedPolygonIDs )
+								{
+									// We only care about polygons that are members of the set of polygons we were asked to extrude
+									if( PolygonsSet.Contains( ConnectedPolygonID ) )
+									{
+										NeighborPolygonIDs.Add( ConnectedPolygonID );
+
+										// We'll need this polygon's normal to figure out where to put the extruded copy of the polygon
+										const FVector NeighborPolygonNormal = ComputePolygonNormal( ConnectedPolygonID );
+										ExtrudeDirection += NeighborPolygonNormal;
+									}
+								}
+								ExtrudeDirection.Normalize();
+
+
+								// OK, we have the direction to extrude for this vertex.  Now we need to know how far to extrude.  We'll
+								// loop over all of the neighbor polygons to this vertex, and choose the closest intersection point with our
+								// vertex's extrude direction and the neighbor polygon's extruded plane
+								FVector ClosestIntersectionPointWithExtrudedPlanes;
+								float ClosestIntersectionDistanceSquared = TNumericLimits<float>::Max();
+
+								for( const FPolygonID NeighborPolygonID : NeighborPolygonIDs )
+								{
+									const FPlane NeighborPolygonPlane = ComputePolygonPlane( NeighborPolygonID );
+
+									// Push the plane out
+									const FPlane ExtrudedPlane = [NeighborPolygonPlane, ExtrudeDistance]
+										{ 
+											FPlane NewPlane = NeighborPolygonPlane;
+											NewPlane.W += ExtrudeDistance; 
+											return NewPlane; 
+										}();
+
+									// Is this the closest intersection point so far?
+									const FVector IntersectionPointWithExtrudedPlane = FMath::RayPlaneIntersection( OriginalVertexPosition, ExtrudeDirection, ExtrudedPlane );
+									const float IntersectionDistanceSquared = FVector::DistSquared( OriginalVertexPosition, IntersectionPointWithExtrudedPlane );
+									if( IntersectionDistanceSquared < ClosestIntersectionDistanceSquared )
+									{
+										ClosestIntersectionPointWithExtrudedPlanes = IntersectionPointWithExtrudedPlane;
+										ClosestIntersectionDistanceSquared = IntersectionDistanceSquared;
+									}
+								}
+
+								ExtrudedVertexPosition = ClosestIntersectionPointWithExtrudedPlanes;
+							}
+							else
+							{
+								// We'll need this polygon's normal to figure out where to put the extruded copy of the polygon
+								const FVector PolygonNormal = ComputePolygonNormal( PolygonID );
+								ExtrudedVertexPosition = OriginalVertexPosition + ExtrudeDistance * PolygonNormal;
+							}
+
+							// Fill in the vertex
+							FAttributesForVertex& AttributesForVertex = *new( AttributesForVertices ) FAttributesForVertex();
+							AttributesForVertex.VertexID = ExtrudedVertexID;
+							AttributesForVertex.VertexAttributes.Attributes.Add( FMeshElementAttributeData(
+								UEditableMeshAttribute::VertexPosition(),
+								0,
+								ExtrudedVertexPosition ) );
+						}
+						ExtrudedEdgeVertexIDs[ EdgeVertexNumber ] = *ExtrudedEdgeVertexIDPtr;
+					}
+
+					if( !bIsSharedEdge )
+					{
+						// Add a face that connects the edge to it's extruded counterpart edge
+						FVertexID OriginalVertexIDs[ 4 ];
+
+						static TArray<FVertexAndAttributes> NewSidePolygonVertices;
+						NewSidePolygonVertices.Reset( 4 );
+						NewSidePolygonVertices.SetNum( 4, false );	// Always four edges in an extruded face
+
+						NewSidePolygonVertices[ 0 ].VertexID = EdgeVertexIDs[ 1 ];
+						OriginalVertexIDs[ 0 ] = EdgeVertexIDs[ 1 ];
+						NewSidePolygonVertices[ 1 ].VertexID = EdgeVertexIDs[ 0 ];
+						OriginalVertexIDs[ 1 ] = EdgeVertexIDs[ 0 ];
+						NewSidePolygonVertices[ 2 ].VertexID = ExtrudedEdgeVertexIDs[ 0 ];
+						OriginalVertexIDs[ 2 ] = EdgeVertexIDs[ 0 ];
+						NewSidePolygonVertices[ 3 ].VertexID = ExtrudedEdgeVertexIDs[ 1 ];
+						OriginalVertexIDs[ 3 ] = EdgeVertexIDs[ 1 ];
+
+						FPolygonID NewSidePolygonID;	// Filled in below
+						{
+							static TArray<FPolygonToCreate> PolygonsToCreate;
+							PolygonsToCreate.Reset();
+
+							// Create the polygon
+							// @todo mesheditor perf: Ideally we support creating multiple polygons at once and batching up the work
+							FPolygonToCreate& PolygonToCreate = *new( PolygonsToCreate ) FPolygonToCreate();
+							PolygonToCreate.PolygonGroupID = PolygonGroupID;
+							PolygonToCreate.PerimeterVertices = NewSidePolygonVertices;	// @todo mesheditor perf: Copying static array here, ideally allocations could be avoided
+
+							// NOTE: We never create holes in side polygons
+
+							static TArray<FPolygonID> NewPolygonIDs;
+							NewPolygonIDs.Reset();
+							static TArray<FEdgeID> NewEdgeIDs;
+							NewEdgeIDs.Reset();
+							CreatePolygons( PolygonsToCreate, /* Out */ NewPolygonIDs, /* Out */ NewEdgeIDs );	// @todo mesheditor perf: Extra allocatons/copies: Ideally MoveTemp() here but we can't move a STATIC local!
+
+							NewSidePolygonID = NewPolygonIDs[ 0 ];
+						}
+						AllNewPolygons.Add( NewSidePolygonID );
+
+						// Copy polygon UVs from the original polygon's vertices
+						{
+							FVertexAttributesForPolygon& PolygonNewAttributes = *new( VertexAttributesForPolygons ) FVertexAttributesForPolygon();
+							PolygonNewAttributes.PolygonID = NewSidePolygonID;
+							PolygonNewAttributes.PerimeterVertexAttributeLists.SetNum( NewSidePolygonVertices.Num(), false );
+
+							for( int32 NewVertexIter = 0; NewVertexIter < NewSidePolygonVertices.Num(); ++NewVertexIter )
+							{
+								const int32 PolygonVertexNumber = this->FindPolygonPerimeterVertexNumberForVertex( PolygonID, OriginalVertexIDs[ NewVertexIter ] );
+								check( PolygonVertexNumber != INDEX_NONE );
+
+								TArray<FMeshElementAttributeData>& PerimeterVertexNewAttributes = PolygonNewAttributes.PerimeterVertexAttributeLists[ NewVertexIter ].Attributes;
+								for( int32 TextureCoordinateIndex = 0; TextureCoordinateIndex < TextureCoordinateCount; ++TextureCoordinateIndex )
+								{
+									const FVector4 TextureCoordinate = this->GetPolygonPerimeterVertexAttribute(
+										PolygonID,
+										PolygonVertexNumber,
+										UEditableMeshAttribute::VertexTextureCoordinate(),
+										TextureCoordinateIndex );
+
+									PerimeterVertexNewAttributes.Emplace( UEditableMeshAttribute::VertexTextureCoordinate(), TextureCoordinateIndex, TextureCoordinate );
+								}
+
+								const FVector4 VertexColor = this->GetPolygonPerimeterVertexAttribute(
+									PolygonID,
+									PolygonVertexNumber,
+									UEditableMeshAttribute::VertexColor(),
+									0 );
+
+								PerimeterVertexNewAttributes.Emplace( UEditableMeshAttribute::VertexColor(), 0, VertexColor );
+							}
+						}
+
+						// All of this edges of the new polygon will be hard
+						{
+							const int32 NewPerimeterEdgeCount = this->GetPolygonPerimeterEdgeCount( NewSidePolygonID );
+							for( int32 NewPerimeterEdgeNumber = 0; NewPerimeterEdgeNumber < NewPerimeterEdgeCount; ++NewPerimeterEdgeNumber )
+							{
+								bool bNewEdgeWindingIsReversedForPolygon;
+								const FEdgeID NewEdgeID = this->GetPolygonPerimeterEdge( NewSidePolygonID, NewPerimeterEdgeNumber, /* Out */ bNewEdgeWindingIsReversedForPolygon );
+
+								// New side polygons get hard edges, but not hard crease weights
+								FAttributesForEdge& AttributesForEdge = *new( AttributesForEdges ) FAttributesForEdge();
+								AttributesForEdge.EdgeID = NewEdgeID;
+								AttributesForEdge.EdgeAttributes.Attributes.Add( FMeshElementAttributeData( UEditableMeshAttribute::EdgeIsHard(), 0, FVector4( 1.0f ) ) );
+							}
+						}
+					}
+				}
+			}
+		}
+	}
+
+	for( int32 PolygonIter = 0; PolygonIter < PolygonIDs.Num(); ++PolygonIter )
+	{
+		const FPolygonID PolygonID = PolygonIDs[ PolygonIter ];
+		const FPolygonGroupID PolygonGroupID = GetGroupForPolygon( PolygonID );
+
+		static TArray<FVertexID> PolygonVertexIDs;
+		this->GetPolygonPerimeterVertices( PolygonID, /* Out */ PolygonVertexIDs );
+
+		// Create a new extruded polygon for the face
+		FPolygonID ExtrudedFrontPolygonID;	// Filled in below
+		{
+			static TArray<FVertexAndAttributes> NewFrontPolygonVertices;
+			NewFrontPolygonVertices.Reset( PolygonVertexIDs.Num() );
+			NewFrontPolygonVertices.SetNum( PolygonVertexIDs.Num(), false );
+
+			// Map all of the polygon's vertex IDs to their extruded counterparts to create the new polygon perimeter
+			for( int32 PolygonVertexNumber = 0; PolygonVertexNumber < PolygonVertexIDs.Num(); ++PolygonVertexNumber )
+			{
+				const FVertexID VertexID = PolygonVertexIDs[ PolygonVertexNumber ];
+				const FVertexID* ExtrudedCopyVertexIDPtr = VertexIDToExtrudedCopy.Find( VertexID );
+				if( ExtrudedCopyVertexIDPtr != nullptr )
+				{
+					NewFrontPolygonVertices[ PolygonVertexNumber ].VertexID = VertexIDToExtrudedCopy[ VertexID ];
+				}
+				else
+				{
+					// We didn't need to extrude a new copy of this vertex (because it was part of a shared edge), so just connect the polygon to the original vertex
+					NewFrontPolygonVertices[ PolygonVertexNumber ].VertexID = VertexID;
+				}
+			}
+
+			{
+				static TArray<FPolygonToCreate> PolygonsToCreate;
+				PolygonsToCreate.Reset();
+
+				// Create the polygon
+				// @todo mesheditor perf: Ideally we support creating multiple polygons at once and batching up the work
+				FPolygonToCreate& PolygonToCreate = *new( PolygonsToCreate ) FPolygonToCreate();
+				PolygonToCreate.PolygonGroupID = PolygonGroupID;
+				PolygonToCreate.PolygonEdgeHardness = EPolygonEdgeHardness::AllEdgesHard;
+				PolygonToCreate.PerimeterVertices = NewFrontPolygonVertices;	// @todo mesheditor perf: Copying static array here, ideally allocations could be avoided
+																				// @todo mesheditor holes: This algorithm is ignoring polygon holes!  Should be easy to support this by doing roughly the same thing.
+				static TArray<FPolygonID> NewPolygonIDs;
+				NewPolygonIDs.Reset();
+				static TArray<FEdgeID> NewEdgeIDs;
+				NewEdgeIDs.Reset();
+				CreatePolygons( PolygonsToCreate, /* Out */ NewPolygonIDs, /* Out */ NewEdgeIDs );	// @todo mesheditor perf: Extra allocatons/copies: Ideally MoveTemp() here but we can't move a STATIC local!
+
+				ExtrudedFrontPolygonID = NewPolygonIDs[ 0 ];
+			}
+			AllNewPolygons.Add( ExtrudedFrontPolygonID );
+
+			// Retain the UVs from the original polygon
+			{
+				FVertexAttributesForPolygon& PolygonNewAttributes = *new( VertexAttributesForPolygons ) FVertexAttributesForPolygon();
+				PolygonNewAttributes.PolygonID = ExtrudedFrontPolygonID;
+				PolygonNewAttributes.PerimeterVertexAttributeLists.SetNum( PolygonVertexIDs.Num(), false );
+
+				for( int32 PolygonVertexNumber = 0; PolygonVertexNumber < PolygonVertexIDs.Num(); ++PolygonVertexNumber )
+				{
+					TArray<FMeshElementAttributeData>& PerimeterVertexNewAttributes = PolygonNewAttributes.PerimeterVertexAttributeLists[ PolygonVertexNumber ].Attributes;
+					for( int32 TextureCoordinateIndex = 0; TextureCoordinateIndex < TextureCoordinateCount; ++TextureCoordinateIndex )
+					{
+						const FVector4 TextureCoordinate = this->GetPolygonPerimeterVertexAttribute(
+							PolygonID,
+							PolygonVertexNumber,
+							UEditableMeshAttribute::VertexTextureCoordinate(),
+							TextureCoordinateIndex );
+
+						FMeshElementAttributeData& PerimeterVertexAttribute = *new( PerimeterVertexNewAttributes ) FMeshElementAttributeData( UEditableMeshAttribute::VertexTextureCoordinate(), TextureCoordinateIndex, TextureCoordinate );
+					}
+
+					const FVector4 VertexColor = this->GetPolygonPerimeterVertexAttribute(
+						PolygonID,
+						PolygonVertexNumber,
+						UEditableMeshAttribute::VertexColor(),
+						0 );
+
+					FMeshElementAttributeData& PerimeterVertexAttribute = *new( PerimeterVertexNewAttributes ) FMeshElementAttributeData( UEditableMeshAttribute::VertexColor(), 0, VertexColor );
+				}
+			}
+
+
+			// All of the border edges of the new polygon will be hard.  If it was a shared edge, then we'll just preserve whatever was
+			// originally going on with the internal edge.
+			{
+				const int32 NewPerimeterEdgeCount = this->GetPolygonPerimeterEdgeCount( ExtrudedFrontPolygonID );
+				check( NewPerimeterEdgeCount == GetPolygonPerimeterEdgeCount( PolygonID ) );	// New polygon should always have the same number of edges (in the same order) as the original!
+				for( int32 PerimeterEdgeNumber = 0; PerimeterEdgeNumber < NewPerimeterEdgeCount; ++PerimeterEdgeNumber )
+				{
+					bool bOriginalEdgeWindingIsReversedForPolygon;
+					const FEdgeID OriginalEdgeID = this->GetPolygonPerimeterEdge( PolygonID, PerimeterEdgeNumber, /* Out */ bOriginalEdgeWindingIsReversedForPolygon );
+					const bool bIsSharedEdge = bKeepNeighborsTogether && EdgeUsageCounts[ OriginalEdgeID ] > 1;
+
+					bool bEdgeWindingIsReversedForPolygon;
+					const FEdgeID EdgeID = this->GetPolygonPerimeterEdge( ExtrudedFrontPolygonID, PerimeterEdgeNumber, /* Out */ bEdgeWindingIsReversedForPolygon );
+
+					const FVector4 NewEdgeHardnessAttribute = bIsSharedEdge ? GetEdgeAttribute( OriginalEdgeID, UEditableMeshAttribute::EdgeIsHard(), 0 ) : FVector( 1.0f );
+
+					FAttributesForEdge& AttributesForEdge = *new( AttributesForEdges ) FAttributesForEdge();
+					AttributesForEdge.EdgeID = EdgeID;
+					AttributesForEdge.EdgeAttributes.Attributes.Add( FMeshElementAttributeData( UEditableMeshAttribute::EdgeIsHard(), 0, NewEdgeHardnessAttribute ) );
+				}
+			}
+		}
+
+		OutNewExtrudedFrontPolygons.Add( ExtrudedFrontPolygonID );
+	}
+	check( NextAvailableExtrudedVertexIDNumber == ExtrudedVertexIDs.Num() );	// Make sure all of the vertices we created were actually used by new polygons
+
+	// Update edge attributes in bulk
+	SetEdgesAttributes( AttributesForEdges );
+
+	// Update vertex attributes in bulk
+	SetVerticesAttributes( AttributesForVertices );
+	SetPolygonsVertexAttributes( VertexAttributesForPolygons );
+
+	// Delete the original polygons
+	{
+		const bool bDeleteOrphanedEdges = true;
+		const bool bDeleteOrphanedVertices = true;
+		const bool bDeleteOrphanedVertexInstances = true;
+		const bool bDeleteEmptyPolygonGroups = false;
+		DeletePolygons( PolygonIDs, bDeleteOrphanedEdges, bDeleteOrphanedVertices, bDeleteOrphanedVertexInstances, bDeleteEmptyPolygonGroups );
+	}
+
+	UE_LOG( LogEditableMesh, Verbose, TEXT( "* ExtrudePolygons returned %s" ), *LogHelpers::ArrayToString( OutNewExtrudedFrontPolygons ) );
+}
+
+
+void UEditableMesh::ExtendEdges( const TArray<FEdgeID>& EdgeIDs, const bool bWeldNeighbors, TArray<FEdgeID>& OutNewExtendedEdgeIDs )
+{
+	OutNewExtendedEdgeIDs.Reset();
+
+	static TArray<FVertexID> NewVertexIDs;
+	NewVertexIDs.Reset();
+
+	// For each original edge vertex ID that we'll be creating a counterpart for on the extended edge, a mapping
+	// to the vertex number of our NewVertexIDs (and VerticesToCreate) list.
+	static TMap<FVertexID, int32> OriginalVertexIDToCreatedVertexNumber;
+	OriginalVertexIDToCreatedVertexNumber.Reset();
+
+	// Create new vertices for all of the new edges.  If bWeldNeighbors is true, we'll share vertices between edges that share the
+	// same vertex instead of creating new edges.
+	{
+		static TArray<FVertexToCreate> VerticesToCreate;
+		VerticesToCreate.Reset();
+		VerticesToCreate.Reserve( EdgeIDs.Num() * 2 );	// Might actually end up needing less, but that's OK.
+
+		for( const FEdgeID EdgeID : EdgeIDs )
+		{
+			FVertexID EdgeVertexIDs[ 2 ];
+			GetEdgeVertices( EdgeID, /* Out */ EdgeVertexIDs[0], /* Out */ EdgeVertexIDs[1] );
+
+			for( const FVertexID EdgeVertexID : EdgeVertexIDs )
+			{
+				// Have we already created a counterpart for this vertex?  If we were asked to weld extended neighbor edges,
+				// we'll want to make sure that we share the extended vertex too!
+				const int32* FoundCreatedVertexNumberPtr = OriginalVertexIDToCreatedVertexNumber.Find( EdgeVertexID );
+				if( !( bWeldNeighbors && FoundCreatedVertexNumberPtr != nullptr ) )
+				{
+					const int32 CreatedVertexNumber = VerticesToCreate.Num();
+					FVertexToCreate& VertexToCreate = *new( VerticesToCreate ) FVertexToCreate();
+
+					// Copy attributes from the original vertex
+					for( const FName AttributeName : UEditableMesh::GetValidVertexAttributes() )
+					{
+						const int32 MaxAttributeIndex = GetMaxAttributeIndex( AttributeName );
+						for( int32 AttributeIndex = 0; AttributeIndex < MaxAttributeIndex; ++AttributeIndex )
+						{
+							FMeshElementAttributeData& VertexAttribute = *new( VertexToCreate.VertexAttributes.Attributes ) FMeshElementAttributeData(
+								AttributeName,
+								AttributeIndex,
+								GetVertexAttribute( EdgeVertexID, AttributeName, AttributeIndex ) );
+						}
+					}
+
+					// Keep track of which vertex we're creating a counterpart for
+					// @todo mesheditor perf: In the case where bWeldNeighbors=false, we can skip updating this map and instead calculate
+					// the index of the new vertex (it will be 2*EdgeNumber+EdgeVertexNumber, because all vertices will be unique.)
+					OriginalVertexIDToCreatedVertexNumber.Add( EdgeVertexID, CreatedVertexNumber );
+				}
+			}
+		}
+
+		CreateVertices( VerticesToCreate, /* Out */ NewVertexIDs );
+	}
+
+
+	// Create the extended edges
+	{
+		static TArray<FEdgeToCreate> EdgesToCreate;
+		EdgesToCreate.Reset();
+		EdgesToCreate.Reserve( EdgeIDs.Num() );
+
+		for( const FEdgeID EdgeID : EdgeIDs )
+		{
+			FVertexID EdgeVertexIDs[ 2 ];
+			GetEdgeVertices( EdgeID, /* Out */ EdgeVertexIDs[0], /* Out */ EdgeVertexIDs[1] );
+
+			FEdgeToCreate& EdgeToCreate = *new( EdgesToCreate ) FEdgeToCreate();
+
+			EdgeToCreate.VertexID0 = NewVertexIDs[ OriginalVertexIDToCreatedVertexNumber.FindChecked( EdgeVertexIDs[ 0 ] ) ];
+			EdgeToCreate.VertexID1 = NewVertexIDs[ OriginalVertexIDToCreatedVertexNumber.FindChecked( EdgeVertexIDs[ 1 ] ) ];
+
+			// Copy attributes from our original edge
+			for( const FName AttributeName : UEditableMesh::GetValidEdgeAttributes() )
+			{
+				const int32 MaxAttributeIndex = GetMaxAttributeIndex( AttributeName );
+				for( int32 AttributeIndex = 0; AttributeIndex < MaxAttributeIndex; ++AttributeIndex )
+				{
+					FMeshElementAttributeData& EdgeAttribute = *new( EdgeToCreate.EdgeAttributes.Attributes ) FMeshElementAttributeData(
+						AttributeName,
+						AttributeIndex,
+						GetEdgeAttribute( EdgeID, AttributeName, AttributeIndex ) );
+				}
+			}
+
+			// We're not connected to any polygons yet.  That will come later.
+			EdgeToCreate.ConnectedPolygons.Reset();
+		}
+
+		CreateEdges( EdgesToCreate, /* Out */ OutNewExtendedEdgeIDs );
+	}
+
+
+	// For every edge, make a quad to connect the original edge with it's extended counterpart.
+	{
+		static TArray<FPolygonToCreate> PolygonsToCreate;
+		PolygonsToCreate.Reset();
+		PolygonsToCreate.Reserve( EdgeIDs.Num() );
+
+		for( int32 ExtendedEdgeNumber = 0; ExtendedEdgeNumber < OutNewExtendedEdgeIDs.Num(); ++ExtendedEdgeNumber )
+		{
+			const FEdgeID OriginalEdgeID = EdgeIDs[ ExtendedEdgeNumber ];
+			const FEdgeID ExtendedEdgeID = OutNewExtendedEdgeIDs[ ExtendedEdgeNumber ];
+
+			FVertexID OriginalEdgeVertexIDs[ 2 ];
+			GetEdgeVertices( OriginalEdgeID, /* Out */ OriginalEdgeVertexIDs[0], /* Out */ OriginalEdgeVertexIDs[1] );
+
+			FVertexID ExtendedEdgeVertexIDs[ 2 ];
+			GetEdgeVertices( ExtendedEdgeID, /* Out */ ExtendedEdgeVertexIDs[0], /* Out */ ExtendedEdgeVertexIDs[1] );
+
+			FPolygonToCreate& PolygonToCreate = *new( PolygonsToCreate ) FPolygonToCreate();
+
+			// We need to figure out which mesh PolygonGroup to put the new polygons in.  To do this, we'll look at
+			// which polygons are already connected to the current edge, and use PolygonGroup number from the first
+			// polygon we can find.  If no polygons are connected, then we'll just use first PolygonGroup in the mesh.
+			// We'll also capture texture coordinates from this polygon, so we can apply them to the new polygon vertices.
+			FPolygonID ConnectedPolygonID = FPolygonID::Invalid;
+			{
+				const int32 ConnectedPolygonCount = GetEdgeConnectedPolygonCount( OriginalEdgeID );
+				if( ConnectedPolygonCount > 0 )
+				{
+					ConnectedPolygonID = GetEdgeConnectedPolygon( OriginalEdgeID, 0 );
+				}
+			}
+
+			PolygonToCreate.PolygonGroupID = ( ConnectedPolygonID != FPolygonID::Invalid ) ? GetGroupForPolygon( ConnectedPolygonID ) : GetFirstValidPolygonGroup();
+			check( PolygonToCreate.PolygonGroupID != FPolygonGroupID::Invalid );
+
+			PolygonToCreate.PerimeterVertices.SetNum( 4, false );
+			FVertexID ConnectedPolygonVertexIDsForTextureCoordinates[ 4 ];
+			{
+				// @todo mesheditor urgent subdiv: This causes degenerate UV triangles, which OpenSubdiv sort of freaks out about (causes flickering geometry)
+				int32 NextVertexNumber = 0;
+				ConnectedPolygonVertexIDsForTextureCoordinates[ NextVertexNumber ] = OriginalEdgeVertexIDs[ 1 ];
+				PolygonToCreate.PerimeterVertices[ NextVertexNumber++ ].VertexID = OriginalEdgeVertexIDs[ 1 ];
+
+				ConnectedPolygonVertexIDsForTextureCoordinates[ NextVertexNumber ] = OriginalEdgeVertexIDs[ 0 ];
+				PolygonToCreate.PerimeterVertices[ NextVertexNumber++ ].VertexID = OriginalEdgeVertexIDs[ 0 ];
+
+				ConnectedPolygonVertexIDsForTextureCoordinates[ NextVertexNumber ] = OriginalEdgeVertexIDs[ 0 ];
+				PolygonToCreate.PerimeterVertices[ NextVertexNumber++ ].VertexID = ExtendedEdgeVertexIDs[ 0 ];
+
+				ConnectedPolygonVertexIDsForTextureCoordinates[ NextVertexNumber ] = OriginalEdgeVertexIDs[ 1 ];
+				PolygonToCreate.PerimeterVertices[ NextVertexNumber++ ].VertexID = ExtendedEdgeVertexIDs[ 1 ];
+			}
+
+			// Preserve polygon vertex attributes
+			if( ConnectedPolygonID != FPolygonID::Invalid )
+			{
+				for( int32 PerimeterVertexNumber = 0; PerimeterVertexNumber < PolygonToCreate.PerimeterVertices.Num(); ++PerimeterVertexNumber )
+				{
+					const FVertexID ConnectedPolygonVertexIDForTextureCoordinates = ConnectedPolygonVertexIDsForTextureCoordinates[ PerimeterVertexNumber ];
+
+					TArray<FMeshElementAttributeData>& AttributesForVertex = PolygonToCreate.PerimeterVertices[ PerimeterVertexNumber ].PolygonVertexAttributes.Attributes;
+				
+					const int32 ConnectedPolygonPerimeterVertexNumber = FindPolygonPerimeterVertexNumberForVertex( ConnectedPolygonID, ConnectedPolygonVertexIDForTextureCoordinates );
+					check( ConnectedPolygonPerimeterVertexNumber != INDEX_NONE );
+
+					for( int32 TextureCoordinateIndex = 0; TextureCoordinateIndex < TextureCoordinateCount; ++TextureCoordinateIndex )
+					{
+						const FVector4 TextureCoordinate = GetPolygonPerimeterVertexAttribute( ConnectedPolygonID, ConnectedPolygonPerimeterVertexNumber, UEditableMeshAttribute::VertexTextureCoordinate(), TextureCoordinateIndex );
+						AttributesForVertex.Add( FMeshElementAttributeData( UEditableMeshAttribute::VertexTextureCoordinate(), TextureCoordinateIndex, TextureCoordinate ) );
+					}
+
+					const FVector4 VertexColor = GetPolygonPerimeterVertexAttribute( ConnectedPolygonID, ConnectedPolygonPerimeterVertexNumber, UEditableMeshAttribute::VertexColor(), 0 );
+					AttributesForVertex.Add( FMeshElementAttributeData( UEditableMeshAttribute::VertexColor(), 0, VertexColor ) );
+				}
+			}
+		}
+
+		// Create the polygons.  Note that this will also automatically create the missing side edges that connect
+		// the original edge to it's extended counterpart.
+		static TArray<FPolygonID> NewPolygonIDs;
+		NewPolygonIDs.Reset();
+		static TArray<FEdgeID> NewEdgeIDs;
+		NewEdgeIDs.Reset();
+		CreatePolygons( PolygonsToCreate, /* Out */ NewPolygonIDs, /* Out */ NewEdgeIDs );
+
+		// Expecting no more than two new edges to be created while creating the polygon.  It's possible for zero or one edge
+		// to be created, depending on how many edges we share with neighbors that were extended.
+		check( bWeldNeighbors ? ( NewEdgeIDs.Num() <= 2 * EdgeIDs.Num() ) : ( NewEdgeIDs.Num() == 2 * EdgeIDs.Num() ) );
+	}
+}
+
+
+void UEditableMesh::ExtendVertices( const TArray<FVertexID>& VertexIDs, const bool bOnlyExtendClosestEdge, const FVector ReferencePosition, TArray<FVertexID>& OutNewExtendedVertexIDs )
+{
+	OutNewExtendedVertexIDs.Reset();
+
+	// Create new vertices for all of the new edges.  If bWeldNeighbors is true, we'll share vertices between edges that share the
+	// same vertex instead of creating new edges.
+	{
+		static TArray<FVertexToCreate> VerticesToCreate;
+		VerticesToCreate.Reset();
+		VerticesToCreate.Reserve( VertexIDs.Num() );
+
+		for( const FVertexID VertexID : VertexIDs )
+		{
+			const int32 CreatedVertexNumber = VerticesToCreate.Num();
+			FVertexToCreate& VertexToCreate = *new( VerticesToCreate ) FVertexToCreate();
+
+			// Copy attributes from the original vertex
+			for( const FName AttributeName : UEditableMesh::GetValidVertexAttributes() )
+			{
+				const int32 MaxAttributeIndex = GetMaxAttributeIndex( AttributeName );
+				for( int32 AttributeIndex = 0; AttributeIndex < MaxAttributeIndex; ++AttributeIndex )
+				{
+					FMeshElementAttributeData& VertexAttribute = *new( VertexToCreate.VertexAttributes.Attributes ) FMeshElementAttributeData(
+						AttributeName,
+						AttributeIndex,
+						GetVertexAttribute( VertexID, AttributeName, AttributeIndex ) );
+				}
+			}
+		}
+
+		CreateVertices( VerticesToCreate, /* Out */ OutNewExtendedVertexIDs );
+	}
+
+
+	// For each vertex, we'll now create new triangles to connect the new vertex to each of the original vertex's adjacent vertices.
+	// If the option bOnlyExtendClosestEdge was enabled, we'll only bother doing this for next closest vertex (so, only a single
+	// triangle per vertex will be created.)
+	{
+		static TArray<FPolygonToCreate> PolygonsToCreate;
+		PolygonsToCreate.Reset();
+
+		for( int32 VertexNumber = 0; VertexNumber < VertexIDs.Num(); ++VertexNumber )
+		{
+			const FVertexID OriginalVertexID = VertexIDs[ VertexNumber ];
+			const FVertexID NewVertexID = OutNewExtendedVertexIDs[ VertexNumber ];
+
+			FVertexID ClosestVertexID = FVertexID::Invalid;
+			if( bOnlyExtendClosestEdge )
+			{
+				// Iterate over the edges connected to this vertex, and figure out which edge is closest to the 
+				// specified reference position
+				float ClosestSquaredEdgeDistance = TNumericLimits<float>::Max();
+
+				const int32 ConnectedEdgeCount = GetVertexConnectedEdgeCount( OriginalVertexID );
+				for( int32 EdgeNumber = 0; EdgeNumber < ConnectedEdgeCount; ++EdgeNumber )
+				{
+					const FEdgeID ConnectedEdgeID = GetVertexConnectedEdge( OriginalVertexID, EdgeNumber );
+
+					FVertexID EdgeVertexIDs[ 2 ];
+					GetEdgeVertices( ConnectedEdgeID, /* Out */ EdgeVertexIDs[0], /* Out */ EdgeVertexIDs[1] );
+
+					FVector EdgeVertexPositions[ 2 ];
+					for( int32 EdgeVertexNumber = 0; EdgeVertexNumber < 2; ++EdgeVertexNumber )
+					{
+						EdgeVertexPositions[ EdgeVertexNumber ] = GetVertexAttribute( EdgeVertexIDs[ EdgeVertexNumber ], UEditableMeshAttribute::VertexPosition(), 0 );
+					}
+
+					const float SquaredEdgeDistance = FMath::PointDistToSegmentSquared( ReferencePosition, EdgeVertexPositions[ 0 ], EdgeVertexPositions[ 1 ] );
+					if( SquaredEdgeDistance < ClosestSquaredEdgeDistance )
+					{
+						ClosestVertexID = EdgeVertexIDs[ 0 ] == OriginalVertexID ? EdgeVertexIDs[ 1 ] : EdgeVertexIDs[ 0 ];
+						ClosestSquaredEdgeDistance = SquaredEdgeDistance;
+					}
+				}
+			}
+
+			static TArray<FVertexID> AdjacentVertexIDs;
+			GetVertexAdjacentVertices( OriginalVertexID, /* Out */ AdjacentVertexIDs );
+
+			// For every adjacent vertex, go ahead and create a new triangle
+			for( const FVertexID AdjacentVertexID : AdjacentVertexIDs )
+			{
+				// If we were asked to only extend an edge that's closest to a reference position, check for that here
+				if( !bOnlyExtendClosestEdge || ( AdjacentVertexID == ClosestVertexID ) )
+				{
+					FPolygonToCreate& PolygonToCreate = *new( PolygonsToCreate ) FPolygonToCreate();
+
+					// Figure out which of the connected polygons shares the edge we're going to be using
+					FPolygonID ConnectedPolygonID = FPolygonID::Invalid;
+					{
+						static TArray<FPolygonID> ConnectedPolygonIDs;
+						GetVertexConnectedPolygons( OriginalVertexID, /* Out */ ConnectedPolygonIDs );
+
+						for( const FPolygonID PolygonID : ConnectedPolygonIDs )
+						{
+							const int32 AdjacentVertexNumber = FindPolygonPerimeterVertexNumberForVertex( PolygonID, AdjacentVertexID );
+							if( AdjacentVertexNumber != INDEX_NONE )
+							{
+								// NOTE: There can be more than one polygon that shares this edge.  We'll just take the first one we find.
+								ConnectedPolygonID = PolygonID;
+								break;
+							}
+						}
+					}
+
+					bool bConnectedPolygonWindsForward = true;
+					if( ConnectedPolygonID != FPolygonID::Invalid )
+					{
+						const int32 OriginalVertexNumberOnConnectedPolygon = FindPolygonPerimeterVertexNumberForVertex( ConnectedPolygonID, OriginalVertexID );
+						check( OriginalVertexNumberOnConnectedPolygon != INDEX_NONE );
+
+						const int32 AdjacentVertexNumberOnConnectedPolygon = FindPolygonPerimeterVertexNumberForVertex( ConnectedPolygonID, AdjacentVertexID );
+						check( AdjacentVertexNumberOnConnectedPolygon != INDEX_NONE );
+
+						const int32 PerimeterVertexCount = GetPolygonPerimeterVertexCount( ConnectedPolygonID );
+						if( !( OriginalVertexNumberOnConnectedPolygon == ( PerimeterVertexCount - 1 ) && AdjacentVertexNumberOnConnectedPolygon == 0 ) &&  // Doesn't wrap forward?
+							( OriginalVertexNumberOnConnectedPolygon > AdjacentVertexNumberOnConnectedPolygon ||			// Winds backwards?
+							  AdjacentVertexNumberOnConnectedPolygon == ( PerimeterVertexCount - 1 ) && OriginalVertexNumberOnConnectedPolygon == 0 ) )	// Wraps backwards?
+						{
+							bConnectedPolygonWindsForward = false;
+						}
+					}
+
+					PolygonToCreate.PolygonGroupID = ( ConnectedPolygonID != FPolygonID::Invalid ) ? GetGroupForPolygon( ConnectedPolygonID ) : GetFirstValidPolygonGroup();
+					check( PolygonToCreate.PolygonGroupID != FPolygonGroupID::Invalid );
+
+					FVertexID ConnectedPolygonVertexIDsForTextureCoordinates[ 3 ];
+					PolygonToCreate.PerimeterVertices.SetNum( 3, false );
+					{
+						int32 NextVertexNumber = 0;
+
+						// Original selected vertex
+						ConnectedPolygonVertexIDsForTextureCoordinates[ NextVertexNumber ] = OriginalVertexID;
+						PolygonToCreate.PerimeterVertices[ NextVertexNumber++ ].VertexID = OriginalVertexID;
+
+						if( bConnectedPolygonWindsForward )
+						{
+							// The new vertex we created
+							ConnectedPolygonVertexIDsForTextureCoordinates[ NextVertexNumber ] = OriginalVertexID;
+							PolygonToCreate.PerimeterVertices[ NextVertexNumber++ ].VertexID = NewVertexID;
+
+							// The adjacent vertex
+							ConnectedPolygonVertexIDsForTextureCoordinates[ NextVertexNumber ] = AdjacentVertexID;
+							PolygonToCreate.PerimeterVertices[ NextVertexNumber++ ].VertexID = AdjacentVertexID;
+						}
+						else
+						{
+							// The adjacent vertex
+							ConnectedPolygonVertexIDsForTextureCoordinates[ NextVertexNumber ] = AdjacentVertexID;
+							PolygonToCreate.PerimeterVertices[ NextVertexNumber++ ].VertexID = AdjacentVertexID;
+
+							// The new vertex we created
+							ConnectedPolygonVertexIDsForTextureCoordinates[ NextVertexNumber ] = OriginalVertexID;
+							PolygonToCreate.PerimeterVertices[ NextVertexNumber++ ].VertexID = NewVertexID;
+						}
+					}
+
+					// Preserve polygon vertex attributes
+					if( ConnectedPolygonID != FPolygonID::Invalid )
+					{
+						for( int32 PerimeterVertexNumber = 0; PerimeterVertexNumber < PolygonToCreate.PerimeterVertices.Num(); ++PerimeterVertexNumber )
+						{
+							const FVertexID ConnectedPolygonVertexIDForTextureCoordinates = ConnectedPolygonVertexIDsForTextureCoordinates[ PerimeterVertexNumber ];
+
+							TArray<FMeshElementAttributeData>& AttributesForVertex = PolygonToCreate.PerimeterVertices[ PerimeterVertexNumber ].PolygonVertexAttributes.Attributes;
+
+							const int32 ConnectedPolygonPerimeterVertexNumber = FindPolygonPerimeterVertexNumberForVertex( ConnectedPolygonID, ConnectedPolygonVertexIDForTextureCoordinates );
+							check( ConnectedPolygonPerimeterVertexNumber != INDEX_NONE );
+
+							for( int32 TextureCoordinateIndex = 0; TextureCoordinateIndex < TextureCoordinateCount; ++TextureCoordinateIndex )
+							{
+								const FVector4 TextureCoordinate = GetPolygonPerimeterVertexAttribute( ConnectedPolygonID, ConnectedPolygonPerimeterVertexNumber, UEditableMeshAttribute::VertexTextureCoordinate(), TextureCoordinateIndex );
+								AttributesForVertex.Add( FMeshElementAttributeData( UEditableMeshAttribute::VertexTextureCoordinate(), TextureCoordinateIndex, TextureCoordinate ) );
+							}
+
+							const FVector4 VertexColor = GetPolygonPerimeterVertexAttribute( ConnectedPolygonID, ConnectedPolygonPerimeterVertexNumber, UEditableMeshAttribute::VertexColor(), 0 );
+							AttributesForVertex.Add( FMeshElementAttributeData( UEditableMeshAttribute::VertexColor(), 0, VertexColor ) );
+						}
+					}
+				}
+			}		
+		}
+
+		static TArray<FPolygonID> NewPolygonIDs;
+		NewPolygonIDs.Reset();
+
+		// Create the polygons.  Note that this will also automatically create the missing side edges that connect
+		// the original edge to it's extended counterpart.
+		static TArray<FEdgeID> NewEdgeIDs;
+		NewEdgeIDs.Reset();
+		CreatePolygons( PolygonsToCreate, /* Out */ NewPolygonIDs, /* Out */ NewEdgeIDs );
+	}
+}
+
+
+void UEditableMesh::ComputePolygonsSharedEdges( const TArray<FPolygonID>& PolygonIDs, TArray<FEdgeID>& OutSharedEdgeIDs ) const
+{
+	OutSharedEdgeIDs.Reset();
+
+	static TSet<FEdgeID> EdgesSeenSoFar;
+	EdgesSeenSoFar.Reset();
+	for( const FPolygonID PolygonID : PolygonIDs )
+	{
+		static TArray<FEdgeID> PolygonPerimeterEdgeIDs;
+		GetPolygonPerimeterEdges( PolygonID, /* Out */ PolygonPerimeterEdgeIDs );
+
+		for( const FEdgeID EdgeID : PolygonPerimeterEdgeIDs )
+		{
+			bool bWasAlreadyInSet = false;
+			EdgesSeenSoFar.Add( EdgeID, /* Out */ &bWasAlreadyInSet );
+
+			if( bWasAlreadyInSet )
+			{
+				// OK, this edge was referenced by more than one polygon!
+				OutSharedEdgeIDs.Add( EdgeID );
+			}
+		}
+	}
+}
+
+
+void UEditableMesh::BevelOrInsetPolygons( const TArray<FPolygonID>& PolygonIDs, const float InsetFixedDistance, const float InsetProgressTowardCenter, const EInsetPolygonsMode Mode, const bool bShouldBevel, TArray<FPolygonID>& OutNewCenterPolygonIDs, TArray<FPolygonID>& OutNewSidePolygonIDs )
+{
+	// @todo mesheditor inset/bevel: Weird feedback loop issues at glancing angles when moving mouse while beveling
+
+	static TArray<FPolygonToCreate> SidePolygonsToCreate;
+	SidePolygonsToCreate.Reset();
+
+	static TArray<FPolygonToCreate> CenterPolygonsToCreate;
+	CenterPolygonsToCreate.Reset();
+
+	static TArray<FAttributesForVertex> AttributesForVertices;
+	AttributesForVertices.Reset();
+
+	for( const FPolygonID PolygonID : PolygonIDs )
+	{
+		const FPolygonGroupID PolygonGroupID = GetGroupForPolygon( PolygonID );
+
+		// Find the center of this polygon
+		const FVector PolygonCenter = ComputePolygonCenter( PolygonID );
+
+		static TArray<FVertexID> PerimeterVertexIDs;
+		GetPolygonPerimeterVertices( PolygonID, /* Out */ PerimeterVertexIDs );
+
+		static TArray<FVertexToCreate> VerticesToCreate;
+		VerticesToCreate.Reset( PerimeterVertexIDs.Num() );
+
+		static TArray<TArray<FVector4>> TextureCoordinatesForNewVertices;	// @todo mesheditor perf: Extra memory alloc every time because of nested array
+		TextureCoordinatesForNewVertices.Reset();
+		TextureCoordinatesForNewVertices.Reserve( PerimeterVertexIDs.Num() );
+
+		static TArray<FVector4> VertexColorsForNewVertices;
+		VertexColorsForNewVertices.Reset();
+		VertexColorsForNewVertices.Reserve( PerimeterVertexIDs.Num() );
+
+		for( int32 PerimeterVertexNumber = 0; PerimeterVertexNumber < PerimeterVertexIDs.Num(); ++PerimeterVertexNumber )
+		{
+			const FVertexID PerimeterVertexID = PerimeterVertexIDs[ PerimeterVertexNumber ];
+
+			FVector OffsetDirection = FVector::ZeroVector;
+
+			// If we're beveling, go ahead and move the original polygon perimeter vertices
+			if( bShouldBevel )
+			{
+				// Figure out if this vertex is shared with other polygons that we were asked to bevel.  If it is,
+				// then we'll want to offset the vertex in the average direction of all of those shared polygons.
+				// However, if the vertex is ONLY shared with polygons we were asked to bevel (no other polygons),
+				// then we don't need to move it at all -- it's an internal edge vertex.
+				// @todo mesheditor bevel: With multiple polygons, we don't want to move vertices that are part of
+				// a shared internal border edge.  But we're not handling that yet.
+				int32 ConnectedBevelPolygonCount = 0;
+				bool bIsOnlyConnectedToBevelPolygons = true;
+
+				static TArray<FPolygonID> ConnectedPolygonIDs;
+				GetVertexConnectedPolygons( PerimeterVertexID, /* Out */ ConnectedPolygonIDs );
+				for( const FPolygonID ConnectedPolygonID : ConnectedPolygonIDs )
+				{
+					if( PolygonIDs.Contains( ConnectedPolygonID ) )
+					{
+						++ConnectedBevelPolygonCount;
+						const FVector ConnectedPolygonNormal = ComputePolygonNormal( ConnectedPolygonID );
+						OffsetDirection += -ConnectedPolygonNormal;
+					}
+					else
+					{
+						bIsOnlyConnectedToBevelPolygons = false;
+					}
+				}
+
+				OffsetDirection.Normalize();
+			}
+
+
+			const FVector VertexPosition = GetVertexAttribute( PerimeterVertexID, UEditableMeshAttribute::VertexPosition(), 0 );
+
+			FVector DirectionTowardCenter;
+			float DistanceToCenter;
+			( PolygonCenter - VertexPosition ).ToDirectionAndLength( /* Out */ DirectionTowardCenter, /* Out */ DistanceToCenter );
+
+			const float InsetOffset = ( DistanceToCenter * InsetProgressTowardCenter + InsetFixedDistance );
+			const FVector InsetVertexPosition = VertexPosition + DirectionTowardCenter * InsetOffset;
+
+			FVertexToCreate& VertexToCreate = *new( VerticesToCreate ) FVertexToCreate();
+			VertexToCreate.VertexAttributes.Attributes.Add( FMeshElementAttributeData( UEditableMeshAttribute::VertexPosition(), 0, InsetVertexPosition ) );
+
+			// Interpolate the texture coordinates and vertex colors on this polygon to figure out UVs for our new vertex
+			FMeshTriangle Triangle;
+			FVector TriangleVertexWeights;
+
+			if( ComputeBarycentricWeightForPointOnPolygon( PolygonID, InsetVertexPosition, Triangle, TriangleVertexWeights ) )
+			{
+				TArray<FVector4>& TextureCoordinatesForNewVertex = *new( TextureCoordinatesForNewVertices ) TArray<FVector4>();
+				TextureCoordinatesForNewVertex.SetNum( TextureCoordinateCount, false );
+
+				for( int32 TextureCoordinateIndex = 0; TextureCoordinateIndex < TextureCoordinateCount; ++TextureCoordinateIndex )
+				{
+					TextureCoordinatesForNewVertex[ TextureCoordinateIndex ] =
+						TriangleVertexWeights.X * GetVertexInstanceAttribute( Triangle.VertexInstanceID0, UEditableMeshAttribute::VertexTextureCoordinate(), TextureCoordinateIndex ) +
+						TriangleVertexWeights.Y * GetVertexInstanceAttribute( Triangle.VertexInstanceID1, UEditableMeshAttribute::VertexTextureCoordinate(), TextureCoordinateIndex ) +
+						TriangleVertexWeights.Z * GetVertexInstanceAttribute( Triangle.VertexInstanceID2, UEditableMeshAttribute::VertexTextureCoordinate(), TextureCoordinateIndex );
+				}
+
+				VertexColorsForNewVertices.Add(
+					TriangleVertexWeights.X * GetVertexInstanceAttribute( Triangle.VertexInstanceID0, UEditableMeshAttribute::VertexColor(), 0 ) +
+					TriangleVertexWeights.Y * GetVertexInstanceAttribute( Triangle.VertexInstanceID1, UEditableMeshAttribute::VertexColor(), 0 ) +
+					TriangleVertexWeights.Z * GetVertexInstanceAttribute( Triangle.VertexInstanceID2, UEditableMeshAttribute::VertexColor(), 0 )
+				);
+			}
+			else
+			{
+				TArray<FVector4>& TextureCoordinatesForNewVertex = *new( TextureCoordinatesForNewVertices ) TArray<FVector4>();
+
+				// The new vertex must have been slightly outside the bounds of the polygon, probably just floating point error.
+				// Just use zero'd out texture coordinates and white vertex color in this case.
+				TextureCoordinatesForNewVertex.SetNumZeroed( TextureCoordinateCount );
+				VertexColorsForNewVertices.Emplace( 1.0f, 1.0f, 1.0f, 1.0f );
+			}
+
+			// If we're beveling, go ahead and move the original polygon perimeter vertices
+			if( bShouldBevel )
+			{
+				// Offset the vertex by the opposite direction of the polygon's normal.  We'll move it the same distance
+				// that we're insetting the new polygon
+				const FVector NewVertexPosition = VertexPosition + OffsetDirection * InsetOffset;
+
+				// @todo mesheditor urgent: Undo/Redo breaks if we have the same vertex more than once for some reason!!  Get rid of this ideally
+				const bool bAlreadyHaveVertex = AttributesForVertices.ContainsByPredicate( [PerimeterVertexID]( FAttributesForVertex& AttributesForVertex ) { return AttributesForVertex.VertexID == PerimeterVertexID; } );
+				if( !bAlreadyHaveVertex )
+				{
+					FAttributesForVertex& AttributesForVertex = *new( AttributesForVertices ) FAttributesForVertex();
+					AttributesForVertex.VertexID = PerimeterVertexID;
+
+					FMeshElementAttributeData& ElementAttributeData = *new( AttributesForVertex.VertexAttributes.Attributes ) FMeshElementAttributeData();
+					ElementAttributeData.AttributeName = UEditableMeshAttribute::VertexPosition();
+					ElementAttributeData.AttributeValue = NewVertexPosition;
+					ElementAttributeData.AttributeIndex = 0;
+				}
+			}
+		}
+
+		static TArray<FVertexID> NewVertexIDs;
+		CreateVertices( VerticesToCreate, /* Out */ NewVertexIDs );
+
+		// The new (inset) polygon will be surrounded by new "side" quad polygons, one for each vertex of the perimeter
+		// that's being inset.
+		if( Mode == EInsetPolygonsMode::All || Mode == EInsetPolygonsMode::SidePolygonsOnly )
+		{
+			const int32 NewSidePolygonCount = NewVertexIDs.Num();
+			for( int32 SidePolygonNumber = 0; SidePolygonNumber < NewSidePolygonCount; ++SidePolygonNumber )
+			{
+				const int32 LeftVertexNumber = SidePolygonNumber;
+				const int32 RightVertexNumber = ( LeftVertexNumber + 1 ) % NewSidePolygonCount;
+
+				const FVertexID LeftOriginalVertexID = PerimeterVertexIDs[ LeftVertexNumber ];
+				const FVertexID LeftInsetVertexID = NewVertexIDs[ LeftVertexNumber ];
+				const FVertexID RightOriginalVertexID = PerimeterVertexIDs[ RightVertexNumber ];
+				const FVertexID RightInsetVertexID = NewVertexIDs[ RightVertexNumber ];
+
+				FPolygonToCreate& PolygonToCreate = *new( SidePolygonsToCreate ) FPolygonToCreate();
+				PolygonToCreate.PolygonGroupID = PolygonGroupID;
+				if( bShouldBevel )
+				{
+					PolygonToCreate.PolygonEdgeHardness = EPolygonEdgeHardness::AllEdgesHard;
+				}
+
+				// Original left
+				{
+					FVertexAndAttributes& NewPerimeterVertex = *new( PolygonToCreate.PerimeterVertices ) FVertexAndAttributes();
+					NewPerimeterVertex.VertexID = LeftOriginalVertexID;
+
+					// Copy all of the polygon vertex attributes over
+					for( int32 TextureCoordinateIndex = 0; TextureCoordinateIndex < TextureCoordinateCount; ++TextureCoordinateIndex )
+					{
+						NewPerimeterVertex.PolygonVertexAttributes.Attributes.Add( FMeshElementAttributeData(
+							UEditableMeshAttribute::VertexTextureCoordinate(),
+							TextureCoordinateIndex,
+							GetPolygonPerimeterVertexAttribute( PolygonID, LeftVertexNumber, UEditableMeshAttribute::VertexTextureCoordinate(), TextureCoordinateIndex ) ) );
+					}
+
+					NewPerimeterVertex.PolygonVertexAttributes.Attributes.Add( FMeshElementAttributeData(
+						UEditableMeshAttribute::VertexColor(),
+						0,
+						GetPolygonPerimeterVertexAttribute( PolygonID, LeftVertexNumber, UEditableMeshAttribute::VertexColor(), 0 ) ) );
+				}
+
+				// Original right
+				{
+					FVertexAndAttributes& NewPerimeterVertex = *new( PolygonToCreate.PerimeterVertices ) FVertexAndAttributes();
+					NewPerimeterVertex.VertexID = RightOriginalVertexID;
+
+					// Copy all of the polygon vertex attributes over
+					for( int32 TextureCoordinateIndex = 0; TextureCoordinateIndex < TextureCoordinateCount; ++TextureCoordinateIndex )
+					{
+						NewPerimeterVertex.PolygonVertexAttributes.Attributes.Add( FMeshElementAttributeData(
+							UEditableMeshAttribute::VertexTextureCoordinate(),
+							TextureCoordinateIndex,
+							GetPolygonPerimeterVertexAttribute( PolygonID, RightVertexNumber, UEditableMeshAttribute::VertexTextureCoordinate(), TextureCoordinateIndex ) ) );
+					}
+
+					NewPerimeterVertex.PolygonVertexAttributes.Attributes.Add( FMeshElementAttributeData(
+						UEditableMeshAttribute::VertexColor(),
+						0,
+						GetPolygonPerimeterVertexAttribute( PolygonID, RightVertexNumber, UEditableMeshAttribute::VertexColor(), 0 ) ) );
+				}
+
+				// Inset right
+				{
+					FVertexAndAttributes& NewPerimeterVertex = *new( PolygonToCreate.PerimeterVertices ) FVertexAndAttributes();
+					NewPerimeterVertex.VertexID = RightInsetVertexID;
+
+					// Use interpolated texture coordinates
+					for( int32 TextureCoordinateIndex = 0; TextureCoordinateIndex < TextureCoordinateCount; ++TextureCoordinateIndex )
+					{
+						NewPerimeterVertex.PolygonVertexAttributes.Attributes.Add( FMeshElementAttributeData(
+							UEditableMeshAttribute::VertexTextureCoordinate(),
+							TextureCoordinateIndex,
+							TextureCoordinatesForNewVertices[ RightVertexNumber ][ TextureCoordinateIndex ] ) );
+					}
+
+					NewPerimeterVertex.PolygonVertexAttributes.Attributes.Add( FMeshElementAttributeData(
+						UEditableMeshAttribute::VertexColor(),
+						0,
+						VertexColorsForNewVertices[ RightVertexNumber ] ) );
+				}
+
+				// Inset left
+				{
+					FVertexAndAttributes& NewPerimeterVertex = *new( PolygonToCreate.PerimeterVertices ) FVertexAndAttributes();
+					NewPerimeterVertex.VertexID = LeftInsetVertexID;
+
+					// Use interpolated texture coordinates
+					for( int32 TextureCoordinateIndex = 0; TextureCoordinateIndex < TextureCoordinateCount; ++TextureCoordinateIndex )
+					{
+						NewPerimeterVertex.PolygonVertexAttributes.Attributes.Add( FMeshElementAttributeData(
+							UEditableMeshAttribute::VertexTextureCoordinate(),
+							TextureCoordinateIndex,
+							TextureCoordinatesForNewVertices[ LeftVertexNumber ][ TextureCoordinateIndex ] ) );
+					}
+
+					NewPerimeterVertex.PolygonVertexAttributes.Attributes.Add( FMeshElementAttributeData(
+						UEditableMeshAttribute::VertexColor(),
+						0,
+						VertexColorsForNewVertices[ LeftVertexNumber ] ) );
+				}
+			}
+		}
+
+		// Now create the new center polygon that will connect all of the new inset vertices
+		if( Mode == EInsetPolygonsMode::All || Mode == EInsetPolygonsMode::CenterPolygonOnly )
+		{
+			FPolygonToCreate& PolygonToCreate = *new( CenterPolygonsToCreate ) FPolygonToCreate();
+			PolygonToCreate.PolygonGroupID = PolygonGroupID;
+			if( bShouldBevel )
+			{
+				PolygonToCreate.PolygonEdgeHardness = EPolygonEdgeHardness::AllEdgesHard;
+			}
+
+			for( int32 NewVertexNumber = 0; NewVertexNumber < NewVertexIDs.Num(); ++NewVertexNumber )
+			{
+				const FVertexID NewVertexID = NewVertexIDs[ NewVertexNumber ];
+
+				FVertexAndAttributes& NewPerimeterVertex = *new( PolygonToCreate.PerimeterVertices ) FVertexAndAttributes();
+				NewPerimeterVertex.VertexID = NewVertexID;
+
+				// Use interpolated texture coordinates
+				for( int32 TextureCoordinateIndex = 0; TextureCoordinateIndex < TextureCoordinateCount; ++TextureCoordinateIndex )
+				{
+					NewPerimeterVertex.PolygonVertexAttributes.Attributes.Add( FMeshElementAttributeData(
+						UEditableMeshAttribute::VertexTextureCoordinate(),
+						TextureCoordinateIndex,
+						TextureCoordinatesForNewVertices[ NewVertexNumber ][ TextureCoordinateIndex ] ) );
+				}
+
+				NewPerimeterVertex.PolygonVertexAttributes.Attributes.Add( FMeshElementAttributeData(
+					UEditableMeshAttribute::VertexColor(),
+					0,
+					VertexColorsForNewVertices[ NewVertexNumber ] ) );
+			}
+		}
+	}
+
+	// Delete the original polygons
+	{
+		const bool bDeleteOrphanedEdges = false;	// No need to delete orphans, because this function won't orphan anything
+		const bool bDeleteOrphanedVertices = false;
+		const bool bDeleteOrphanedVertexInstances = false;
+		const bool bDeleteEmptyPolygonGroups = false;
+		DeletePolygons( PolygonIDs, bDeleteOrphanedEdges, bDeleteOrphanedVertices, bDeleteOrphanedVertexInstances, bDeleteEmptyPolygonGroups );
+	}
+
+	// @todo mesheditor inset: Should we create the new inset polygon with hard edges or soft?  Currently using the default.
+
+	// Updated any vertices that need to be moved
+	if( AttributesForVertices.Num() > 0 )
+	{
+		SetVerticesAttributes( AttributesForVertices );
+	}
+
+	if( SidePolygonsToCreate.Num() > 0 )
+	{
+		static TArray<FEdgeID> NewEdgeIDs;
+		CreatePolygons( SidePolygonsToCreate, /* Out */ OutNewSidePolygonIDs, /* Out */ NewEdgeIDs );
+	}
+
+	// Create the new center polygons.  Note that we pass back the IDs of the newly-created center polygons
+	if( CenterPolygonsToCreate.Num() > 0 )
+	{
+		static TArray<FEdgeID> NewEdgeIDs;
+		CreatePolygons( CenterPolygonsToCreate, /* Out */ OutNewCenterPolygonIDs, /* Out */ NewEdgeIDs );
+	}
+}
+
+
+void UEditableMesh::InsetPolygons( const TArray<FPolygonID>& PolygonIDs, const float InsetFixedDistance, const float InsetProgressTowardCenter, const EInsetPolygonsMode Mode, TArray<FPolygonID>& OutNewCenterPolygonIDs, TArray<FPolygonID>& OutNewSidePolygonIDs )
+{
+	const bool bShouldBevel = false;
+	BevelOrInsetPolygons( PolygonIDs, InsetFixedDistance, InsetProgressTowardCenter, Mode, bShouldBevel, /* Out */ OutNewCenterPolygonIDs, /* Out */ OutNewSidePolygonIDs );
+}
+
+
+void UEditableMesh::BevelPolygons( const TArray<FPolygonID>& PolygonIDs, const float BevelFixedDistance, const float BevelProgressTowardCenter, TArray<FPolygonID>& OutNewCenterPolygonIDs, TArray<FPolygonID>& OutNewSidePolygonIDs )
+{
+	const bool bShouldBevel = true;
+	BevelOrInsetPolygons( PolygonIDs, BevelFixedDistance, BevelProgressTowardCenter, EInsetPolygonsMode::All, bShouldBevel, /* Out */ OutNewCenterPolygonIDs, /* Out */ OutNewSidePolygonIDs );
+}
+
+
+float UEditableMesh::GetPolygonCornerAngleForVertex( const FPolygonID PolygonID, const FVertexID VertexID ) const
+{
+	const FMeshPolygon& Polygon = Polygons[ PolygonID.GetValue() ];
+
+	// Lambda function which returns the inner angle at a given index on a polygon contour
+	auto GetContourAngle = [ this ]( const FMeshPolygonContour& Contour, const int32 ContourIndex )
+	{
+		const int32 NumVertices = Contour.VertexInstanceIDs.Num();
+
+		const FVector PrevVertexPosition = GetVertexInstanceAttribute( Contour.VertexInstanceIDs[ ( ContourIndex + NumVertices - 1 ) % NumVertices ], UEditableMeshAttribute::VertexPosition(), 0 );
+		const FVector ThisVertexPosition = GetVertexInstanceAttribute( Contour.VertexInstanceIDs[ ContourIndex ], UEditableMeshAttribute::VertexPosition(), 0 );
+		const FVector NextVertexPosition = GetVertexInstanceAttribute( Contour.VertexInstanceIDs[ ( ContourIndex + 1 ) % NumVertices ], UEditableMeshAttribute::VertexPosition(), 0 );
+
+		const FVector Direction1 = ( PrevVertexPosition - ThisVertexPosition ).GetSafeNormal();
+		const FVector Direction2 = ( NextVertexPosition - ThisVertexPosition ).GetSafeNormal();
+
+		return FMath::Acos( FVector::DotProduct( Direction1, Direction2 ) );
+	};
+
+	auto IsVertexInstancedFromThisVertex = [ this, VertexID ]( const FVertexInstanceID VertexInstanceID )
+	{
+		return VertexInstances[ VertexInstanceID.GetValue() ].VertexID == VertexID;
+	};
+
+	// First look for the vertex instance in the perimeter
+	int32 ContourIndex = Polygon.PerimeterContour.VertexInstanceIDs.IndexOfByPredicate( IsVertexInstancedFromThisVertex );
+	if( ContourIndex != INDEX_NONE )
+	{
+		// Return the internal angle if found
+		return GetContourAngle( Polygon.PerimeterContour, ContourIndex );
+	}
+	else
+	{
+		// If not found, look in all the holes
+		for( const FMeshPolygonContour& HoleContour : Polygon.HoleContours )
+		{
+			ContourIndex = HoleContour.VertexInstanceIDs.IndexOfByPredicate( IsVertexInstancedFromThisVertex );
+			if( ContourIndex != INDEX_NONE )
+			{
+				// Hole vertex contribution is the part which ISN'T the internal angle of the contour, so subtract from 2*pi
+				return ( 2.0f * PI ) - GetContourAngle( HoleContour, ContourIndex );
+			}
+		}
+	}
+
+	// Found nothing; return 0
+	return 0.0f;
+}
+
+
+void UEditableMesh::GeneratePolygonTangentsAndNormals( const TArray<FPolygonID>& PolygonIDs )
+{
+	for( const FPolygonID PolygonID : PolygonIDs )
+	{
+		FMeshPolygon& Polygon = Polygons[ PolygonID.GetValue() ];
+
+		// Calculate the tangent basis for the polygon, based on the average of all constituent triangles
+		FVector Normal = FVector::ZeroVector;
+		FVector Tangent = FVector::ZeroVector;
+		FVector Binormal = FVector::ZeroVector;
+
+		for( const FMeshTriangle& Triangle : Polygon.Triangles )
+		{
+			const FMeshVertexInstance& VertexInstance0 = VertexInstances[ Triangle.VertexInstanceID0.GetValue() ];
+			const FMeshVertexInstance& VertexInstance1 = VertexInstances[ Triangle.VertexInstanceID1.GetValue() ];
+			const FMeshVertexInstance& VertexInstance2 = VertexInstances[ Triangle.VertexInstanceID2.GetValue() ];
+
+			const FMeshVertex& Vertex0 = Vertices[ VertexInstance0.VertexID.GetValue() ];
+			const FMeshVertex& Vertex1 = Vertices[ VertexInstance1.VertexID.GetValue() ];
+			const FMeshVertex& Vertex2 = Vertices[ VertexInstance2.VertexID.GetValue() ];
+
+			const FVector DPosition1 = Vertex1.VertexPosition - Vertex0.VertexPosition;
+			const FVector DPosition2 = Vertex2.VertexPosition - Vertex0.VertexPosition;
+
+			// @todo mesheditor: currently hardcoded to calculate the tangent basis from UV0.
+			const FVector2D DUV1 = VertexInstance1.VertexUVs[ 0 ] - VertexInstance0.VertexUVs[ 0 ];
+			const FVector2D DUV2 = VertexInstance2.VertexUVs[ 0 ] - VertexInstance0.VertexUVs[ 0 ];
+
+			// We have a left-handed coordinate system, but a counter-clockwise winding order
+			// Hence normal calculation has to take the triangle vectors cross product in reverse.
+			Normal += FVector::CrossProduct( DPosition2, DPosition1 );
+
+			// ...and tangent space seems to be right-handed.
+			const float DetUV = FVector2D::CrossProduct( DUV1, DUV2 );
+			const float InvDetUV = ( DetUV == 0.0f ) ? 0.0f : 1.0f / DetUV;
+
+			Tangent += ( DPosition1 * DUV2.Y - DPosition2 * DUV1.Y ) * InvDetUV;
+			Binormal += ( DPosition2 * DUV1.X - DPosition1 * DUV2.X ) * InvDetUV;
+		}
+
+		Polygon.PolygonNormal = Normal.GetSafeNormal();
+		Polygon.PolygonTangent = Tangent.GetSafeNormal();
+		Polygon.PolygonBinormal = Binormal.GetSafeNormal();
+	}
+}
+
+
+void UEditableMesh::GenerateTangentsAndNormals()
+{
+	GeneratePolygonTangentsAndNormals( PolygonsPendingNewTangentBasis.Array() );
+
+	static TSet<FVertexInstanceID> VertexInstanceIDs;
+	VertexInstanceIDs.Reset();
+
+	for( const FPolygonID PolygonID : PolygonsPendingNewTangentBasis )
+	{
+		const FMeshPolygon& Polygon = Polygons[ PolygonID.GetValue() ];
+
+		VertexInstanceIDs.Append( Polygon.PerimeterContour.VertexInstanceIDs );
+		for( const FMeshPolygonContour& HoleContour : Polygon.HoleContours )
+		{
+			VertexInstanceIDs.Append( HoleContour.VertexInstanceIDs );
+		}
+	}
+
+	static TArray<FAttributesForVertexInstance> AttributesForVertexInstances;
+	AttributesForVertexInstances.Reset( VertexInstanceIDs.Num() );
+
+	for( const FVertexInstanceID VertexInstanceID : VertexInstanceIDs )
+	{
+		FVector Normal = FVector::ZeroVector;
+		FVector Tangent = FVector::ZeroVector;
+		FVector Binormal = FVector::ZeroVector;
+
+		FMeshVertexInstance& VertexInstance = VertexInstances[ VertexInstanceID.GetValue() ];
+		const FVertexID VertexID = VertexInstance.VertexID;
+
+		// Get all polygons connected to this vertex instance, and also any in the same smoothing group connected to a different vertex instance
+		// (as they still have influence over the normal).
+		static TArray<FPolygonID> AllConnectedPolygons;
+		check( VertexInstance.ConnectedPolygons.Num() > 0 );
+		GetPolygonsInSameSoftEdgedGroup( VertexID, VertexInstance.ConnectedPolygons[ 0 ], AllConnectedPolygons );
+
+		// The vertex instance normal is computed as a sum of all connected polygons' normals, weighted by the angle they make with the vertex
+		for( const FPolygonID ConnectedPolygonID : AllConnectedPolygons )
+		{
+			float Angle = GetPolygonCornerAngleForVertex( ConnectedPolygonID, VertexID );
+
+			const FMeshPolygon& Polygon = Polygons[ ConnectedPolygonID.GetValue() ];
+			Normal += Polygon.PolygonNormal * Angle;
+
+			// If this polygon is actually connected to the vertex instance we're processing, also include its contributions towards the tangent
+			if( VertexInstance.ConnectedPolygons.Contains( ConnectedPolygonID ) )
+			{
+				Tangent += Polygon.PolygonTangent * Angle;
+				Binormal += Polygon.PolygonBinormal * Angle;
+			}
+		}
+
+		// Normalize Normal
+		Normal = Normal.GetSafeNormal();
+
+		// Make Tangent orthonormal to Normal.
+		// This is a quicker method than normalizing Tangent, taking the cross product Normal X Tangent, and then a further cross product with that result
+		Tangent = ( Tangent - Normal * FVector::DotProduct( Normal, Tangent ) ).GetSafeNormal();
+
+		// Calculate binormal sign
+		const float BinormalSign = ( FVector::DotProduct( FVector::CrossProduct( Normal, Tangent ), Binormal ) < 0.0f ) ? -1.0f : 1.0f;
+
+		//UE_LOG( LogEditableMesh, Verbose, TEXT( "%s normal: old %s" ), *VertexInstanceID.ToString(), *VertexInstance.Normal.ToString() );
+		//UE_LOG( LogEditableMesh, Verbose, TEXT( "%s normal: new %s" ), *VertexInstanceID.ToString(), *Normal.ToString() );
+
+		//UE_LOG( LogEditableMesh, Verbose, TEXT( "%s tangent: old %s" ), *VertexInstanceID.ToString(), *VertexInstance.Tangent.ToString() );
+		//UE_LOG( LogEditableMesh, Verbose, TEXT( "%s tangent: new %s" ), *VertexInstanceID.ToString(), *Tangent.ToString() );
+
+		AttributesForVertexInstances.Emplace();
+		FAttributesForVertexInstance& AttributesForVertexInstance = AttributesForVertexInstances.Last();
+
+		AttributesForVertexInstance.VertexInstanceID = VertexInstanceID;
+		AttributesForVertexInstance.VertexInstanceAttributes.Attributes.Reset( 3 );
+		AttributesForVertexInstance.VertexInstanceAttributes.Attributes.Emplace( UEditableMeshAttribute::VertexNormal(), 0, FVector4( Normal ) );
+		AttributesForVertexInstance.VertexInstanceAttributes.Attributes.Emplace( UEditableMeshAttribute::VertexTangent(), 0, FVector4( Tangent ) );
+		AttributesForVertexInstance.VertexInstanceAttributes.Attributes.Emplace( UEditableMeshAttribute::VertexBinormalSign(), 0, FVector4( BinormalSign, 0.0f, 0.0f, 0.0f ) );
+	}
+
+	SetVertexInstancesAttributes( AttributesForVertexInstances );
+}
+
+
+#if 0
+void UEditableMesh::GenerateTangentsForPolygons( const TArray< FPolygonID >& PolygonIDs )
+{
+	// @todo mesheditor: This method is currently not used. Revisit MikkTSpace in the future.
+	static TArray<FVertexAttributesForPolygon> VertexAttributesForPolygons;
+	VertexAttributesForPolygons.Reset();
+	VertexAttributesForPolygons.Reserve( PolygonIDs.Num() );
+	for( const FPolygonID PolygonID : PolygonIDs )
+	{
+		FVertexAttributesForPolygon& PolygonNewAttributes = *new( VertexAttributesForPolygons ) FVertexAttributesForPolygon();
+		PolygonNewAttributes.PolygonID = PolygonID;
+
+		const int32 PolygonPerimeterVertexCount = GetPolygonPerimeterVertexCount( PolygonID );
+		PolygonNewAttributes.PerimeterVertexAttributeLists.SetNum( PolygonPerimeterVertexCount, false );
+	}
+
+	struct FMikkUserData
+	{
+		UEditableMesh* Self;
+		const TArray< FPolygonID >& Polygons;
+		TArray<FVertexAttributesForPolygon>& VertexAttributesForPolygons;
+
+		FMikkUserData( UEditableMesh* InitSelf, const TArray< FPolygonID >& InitPolygons, TArray<FVertexAttributesForPolygon>& InitVertexAttributesForPolygons )
+			: Self( InitSelf ),
+			  Polygons( InitPolygons ),
+			  VertexAttributesForPolygons( InitVertexAttributesForPolygons )
+		{
+		}
+	} MikkUserData( this, PolygonIDs, VertexAttributesForPolygons );
+
+	struct Local
+	{
+		static int MikkGetNumFaces( const SMikkTSpaceContext* Context )
+		{
+			const FMikkUserData& UserData = *static_cast<FMikkUserData*>( Context->m_pUserData );
+			return UserData.Polygons.Num();
+		}
+
+		static int MikkGetNumVertsOfFace( const SMikkTSpaceContext* Context, const int MikkFaceIndex )
+		{
+			const FMikkUserData& UserData = *static_cast<FMikkUserData*>( Context->m_pUserData );
+
+			const FPolygonID PolygonID = UserData.Polygons[ MikkFaceIndex ];
+
+			// @todo mesheditor holes: Can we even use Mikktpsace for polygons with holes?  Do we need to run the triangulated triangles through instead?
+			const int32 PolygonVertexCount = UserData.Self->GetPolygonPerimeterVertexCount( PolygonID );
+			return PolygonVertexCount;
+		}
+
+		static void MikkGetPosition( const SMikkTSpaceContext* Context, float OutPosition[3], const int MikkFaceIndex, const int MikkVertexIndex )
+		{
+			const FMikkUserData& UserData = *static_cast<FMikkUserData*>( Context->m_pUserData );
+
+			const FPolygonID PolygonID = UserData.Polygons[ MikkFaceIndex ];
+			const FVector VertexPosition = UserData.Self->GetPolygonPerimeterVertexAttribute( PolygonID, MikkVertexIndex, UEditableMeshAttribute::VertexPosition(), 0 );
+
+			OutPosition[0] = VertexPosition.X;
+			OutPosition[1] = VertexPosition.Y;
+			OutPosition[2] = VertexPosition.Z;
+		}
+
+		static void MikkGetNormal( const SMikkTSpaceContext* Context, float OutNormal[3], const int MikkFaceIndex, const int MikkVertexIndex )
+		{
+			const FMikkUserData& UserData = *static_cast<FMikkUserData*>( Context->m_pUserData );
+
+			const FPolygonID PolygonID = UserData.Polygons[ MikkFaceIndex ];
+			const FVector PolygonVertexNormal = UserData.Self->GetPolygonPerimeterVertexAttribute( PolygonID, MikkVertexIndex, UEditableMeshAttribute::VertexNormal(), 0 );
+
+			OutNormal[ 0 ] = PolygonVertexNormal.X;
+			OutNormal[ 1 ] = PolygonVertexNormal.Y;
+			OutNormal[ 2 ] = PolygonVertexNormal.Z;
+		}
+
+		static void MikkGetTexCoord( const SMikkTSpaceContext* Context, float OutUV[2], const int MikkFaceIndex, const int MikkVertexIndex )
+		{
+			const FMikkUserData& UserData = *static_cast<FMikkUserData*>( Context->m_pUserData );
+
+			// @todo mesheditor: Support using a custom texture coordinate index for tangent space generation?
+			const int32 TextureCoordinateIndex = 0;
+
+			const FPolygonID PolygonID = UserData.Polygons[ MikkFaceIndex ];
+			const FVector2D PolygonVertexTextureCoordinate( UserData.Self->GetPolygonPerimeterVertexAttribute( PolygonID, MikkVertexIndex, UEditableMeshAttribute::VertexTextureCoordinate(), TextureCoordinateIndex ) );
+
+			OutUV[0] = PolygonVertexTextureCoordinate.X;
+			OutUV[1] = PolygonVertexTextureCoordinate.Y;
+		}
+
+		static void MikkSetTSpaceBasic( const SMikkTSpaceContext* Context, const float Tangent[3], const float BitangentSign, const int MikkFaceIndex, const int MikkVertexIndex )
+		{
+			FMikkUserData& UserData = *static_cast<FMikkUserData*>( Context->m_pUserData );
+
+			const FVector NewTangent( Tangent[ 0 ], Tangent[ 1 ], Tangent[ 2 ] );
+
+			// Save it!
+			TArray<FMeshElementAttributeData>& PerimeterVertexNewAttributes = UserData.VertexAttributesForPolygons[ MikkFaceIndex ].PerimeterVertexAttributeLists[ MikkVertexIndex ].Attributes;
+			check( PerimeterVertexNewAttributes.Num() == 0 );
+			PerimeterVertexNewAttributes.Emplace( UEditableMeshAttribute::VertexTangent(), 0, FVector4( NewTangent, 0.0f ) );
+			PerimeterVertexNewAttributes.Emplace( UEditableMeshAttribute::VertexBinormalSign(), 0, FVector4( BitangentSign ) );
+		}
+	};
+
+	SMikkTSpaceInterface MikkTInterface;
+	{
+		MikkTInterface.m_getNumFaces = Local::MikkGetNumFaces;
+		MikkTInterface.m_getNumVerticesOfFace = Local::MikkGetNumVertsOfFace;
+		MikkTInterface.m_getPosition = Local::MikkGetPosition;
+		MikkTInterface.m_getNormal = Local::MikkGetNormal;
+		MikkTInterface.m_getTexCoord = Local::MikkGetTexCoord;
+
+		MikkTInterface.m_setTSpaceBasic = Local::MikkSetTSpaceBasic;
+		MikkTInterface.m_setTSpace = nullptr;
+	}
+
+	SMikkTSpaceContext MikkTContext;
+	{
+		MikkTContext.m_pInterface = &MikkTInterface;
+		MikkTContext.m_pUserData = (void*)( &MikkUserData );
+
+		// @todo mesheditor perf: Turning this on can apparently improve performance.  Needs investigation.
+		MikkTContext.m_bIgnoreDegenerates = false;
+	}
+
+	// Now we'll ask MikkTSpace to actually generate the tangents
+	genTangSpaceDefault( &MikkTContext );
+
+	SetPolygonsVertexAttributes( VertexAttributesForPolygons );
+}
+#endif
+
+
+void UEditableMesh::SetVerticesCornerSharpness( const TArray<FVertexID>& VertexIDs, const TArray<float>& VerticesNewSharpness )
+{
+	check( VertexIDs.Num() == VerticesNewSharpness.Num() );
+
+	static TArray<FAttributesForVertex> AttributesForVertices;
+	AttributesForVertices.Reset();
+
+	for( int32 VertexNumber = 0; VertexNumber < VertexIDs.Num(); ++VertexNumber )
+	{
+		const FVertexID VertexID = VertexIDs[ VertexNumber ];
+
+		FAttributesForVertex& AttributesForVertex = *new( AttributesForVertices ) FAttributesForVertex();
+		AttributesForVertex.VertexID = VertexID;
+		AttributesForVertex.VertexAttributes.Attributes.Add( FMeshElementAttributeData(
+			UEditableMeshAttribute::VertexCornerSharpness(),
+			0,
+			FVector4( VerticesNewSharpness[ VertexNumber ] ) ) );
+	}
+
+	SetVerticesAttributes( AttributesForVertices );
+}
+
+
+void UEditableMesh::SetEdgesCreaseSharpness( const TArray<FEdgeID>& EdgeIDs, const TArray<float>& EdgesNewCreaseSharpness )
+{
+	check( EdgeIDs.Num() == EdgesNewCreaseSharpness.Num() );
+
+	static TArray<FAttributesForEdge> AttributesForEdges;
+	AttributesForEdges.Reset();
+
+	for( int32 EdgeNumber = 0; EdgeNumber < EdgeIDs.Num(); ++EdgeNumber )
+	{
+		const FEdgeID EdgeID = EdgeIDs[ EdgeNumber ];
+
+		FAttributesForEdge& AttributesForEdge = *new( AttributesForEdges ) FAttributesForEdge();
+		AttributesForEdge.EdgeID = EdgeID;
+		AttributesForEdge.EdgeAttributes.Attributes.Add( FMeshElementAttributeData(
+			UEditableMeshAttribute::EdgeCreaseSharpness(),
+			0,
+			FVector4( EdgesNewCreaseSharpness[ EdgeNumber ] ) ) );
+	}
+
+	SetEdgesAttributes( AttributesForEdges );
+}
+
+
+void UEditableMesh::SetEdgesHardness( const TArray<FEdgeID>& EdgeIDs, const TArray<bool>& EdgesNewIsHard )
+{
+	check( EdgeIDs.Num() == EdgesNewIsHard.Num() );
+
+	static TArray<FAttributesForEdge> AttributesForEdges;
+	AttributesForEdges.Reset();
+
+	static TSet<FPolygonID> UniqueConnectedPolygonIDs;
+	UniqueConnectedPolygonIDs.Reset();
+
+	for( int32 EdgeNumber = 0; EdgeNumber < EdgeIDs.Num(); ++EdgeNumber )
+	{
+		const FEdgeID EdgeID = EdgeIDs[ EdgeNumber ];
+
+		FAttributesForEdge& AttributesForEdge = *new( AttributesForEdges ) FAttributesForEdge();
+		AttributesForEdge.EdgeID = EdgeID;
+		AttributesForEdge.EdgeAttributes.Attributes.Add( FMeshElementAttributeData(
+			UEditableMeshAttribute::EdgeIsHard(),
+			0,
+			FVector4( EdgesNewIsHard[ EdgeNumber ] ? 1.0f : 0.0f ) ) );
+
+		// Get the polygons this edge is connected to.  They'll need new normals.
+		static TArray<FPolygonID> ConnectedPolygonIDs;
+		GetEdgeConnectedPolygons( EdgeID, /* Out */ ConnectedPolygonIDs );
+		UniqueConnectedPolygonIDs.Append( ConnectedPolygonIDs );
+	}
+
+	SetEdgesAttributes( AttributesForEdges );
+}
+
+
+void UEditableMesh::SetEdgesHardnessAutomatically( const TArray<FEdgeID>& EdgeIDs, const float MaxDotProductForSoftEdge )
+{
+	static TArray<bool> EdgesNewIsHard;
+	EdgesNewIsHard.SetNumUninitialized( EdgeIDs.Num(), false );
+
+	for( int32 EdgeNumber = 0; EdgeNumber < EdgeIDs.Num(); ++EdgeNumber )
+	{
+		const FEdgeID EdgeID = EdgeIDs[ EdgeNumber ];
+
+		// Default to soft if we have no polygons attached
+		bool bIsSoftEdge = true;
+
+		const int32 ConnectedPolygonCount = GetEdgeConnectedPolygonCount( EdgeID );
+		if( ConnectedPolygonCount > 0 )
+		{
+			float MinDot = 1.0f;
+
+			const FPolygonID FirstPolygonID = GetEdgeConnectedPolygon( EdgeID, 0 );
+
+			FVector LastPolygonNormal = ComputePolygonNormal( FirstPolygonID );
+
+			for( int32 ConnectedPolygonNumber = 1; ConnectedPolygonNumber < ConnectedPolygonCount; ++ConnectedPolygonNumber )
+			{
+				const FPolygonID PolygonID = GetEdgeConnectedPolygon( EdgeID, ConnectedPolygonNumber );
+
+				const FVector PolygonNormal = ComputePolygonNormal( PolygonID );
+
+				const float Dot = FVector::DotProduct( PolygonNormal, LastPolygonNormal );
+				MinDot = FMath::Min( Dot, MinDot );
+			}
+
+			bIsSoftEdge = ( MinDot >= MaxDotProductForSoftEdge );
+		}
+
+		EdgesNewIsHard[ EdgeNumber ] = !bIsSoftEdge;
+	}
+
+
+	// Set the edges hardness (and generate new normals)
+	SetEdgesHardness( EdgeIDs, EdgesNewIsHard );
+}
+
+
+void UEditableMesh::SetEdgesVertices( const TArray<FVerticesForEdge>& VerticesForEdges )
+{
+	UE_LOG( LogEditableMesh, Verbose, TEXT( "SetEdgesVertices: %s" ), *LogHelpers::ArrayToString( VerticesForEdges ) );
+
+	FSetEdgesVerticesChangeInput RevertInput;
+
+	RevertInput.VerticesForEdges.AddUninitialized( VerticesForEdges.Num() );
+	for( int32 VertexNumber = 0; VertexNumber < VerticesForEdges.Num(); ++VertexNumber )
+	{
+		const FVerticesForEdge& VerticesForEdge = VerticesForEdges[ VertexNumber ];
+
+		// Save the backup
+		FVerticesForEdge& RevertVerticesForEdge = RevertInput.VerticesForEdges[ VertexNumber ];
+		RevertVerticesForEdge.EdgeID = VerticesForEdge.EdgeID;
+		GetEdgeVertices( VerticesForEdge.EdgeID, /* Out */ RevertVerticesForEdge.NewVertexID0, /* Out */ RevertVerticesForEdge.NewVertexID1 );
+
+		// Edit the edge
+		FMeshEdge& Edge = Edges[ VerticesForEdge.EdgeID.GetValue() ];
+
+		for( uint32 EdgeVertexNumber = 0; EdgeVertexNumber < 2; ++EdgeVertexNumber )
+		{
+			// Disconnect the edge from its existing vertices
+			const FVertexID VertexID = Edge.VertexIDs[ EdgeVertexNumber ];
+			FMeshVertex& Vertex = Vertices[ VertexID.GetValue() ];
+			verify( Vertex.ConnectedEdgeIDs.RemoveSingleSwap( VerticesForEdge.EdgeID ) == 1 );	// Must have been already connected!
+		}
+
+		Edge.VertexIDs[0] = VerticesForEdge.NewVertexID0;
+		Edge.VertexIDs[1] = VerticesForEdge.NewVertexID1;
+
+		// Connect the new vertices to the edge
+		for( uint32 EdgeVertexNumber = 0; EdgeVertexNumber < 2; ++EdgeVertexNumber )
+		{
+			const FVertexID VertexID = Edge.VertexIDs[ EdgeVertexNumber ];
+			FMeshVertex& Vertex = Vertices[ VertexID.GetValue() ];
+
+			check( !Vertex.ConnectedEdgeIDs.Contains( VerticesForEdge.EdgeID ) );	// Should not have already been connected
+			Vertex.ConnectedEdgeIDs.Add( VerticesForEdge.EdgeID );
+		}
+	}
+
+	AddUndo( MakeUnique<FSetEdgesVerticesChange>( MoveTemp( RevertInput ) ) );
+
+	UE_LOG( LogEditableMesh, Verbose, TEXT( "* SetEdgesVertices finished") );
+}
+
+
+void UEditableMesh::InsertPolygonPerimeterVertices( const FPolygonID PolygonID, const int32 InsertBeforeVertexNumber, const TArray<FVertexAndAttributes>& VerticesToInsert )
+{
+	UE_LOG( LogEditableMesh, Verbose, TEXT( "InsertPolygonPerimeterVertices: PolygonID:%s InsertBeforeVertexNumber:%d %s" ), *PolygonID.ToString(), InsertBeforeVertexNumber, *LogHelpers::ArrayToString( VerticesToInsert ) );
+	// @todo mesheditor: see if we want to keep this action. It is only used by SplitEdge, and doesn't generate missing edges.
+	// We can achieve the same thing by deleting a polygon without deleting vertex instances, and creating a new polygon with an extra vertex instance.
+
+	FMeshPolygon& Polygon = Polygons[ PolygonID.GetValue() ];
+
+	// @todo mesheditor: should create missing edges here too?
+
+	// Insert new vertex instances
+	for( int32 InsertVertexIter = 0; InsertVertexIter < VerticesToInsert.Num(); ++InsertVertexIter )
+	{
+		const FVertexAndAttributes& VertexToInsert = VerticesToInsert[ InsertVertexIter ];
+
+		const FVertexInstanceID VertexInstanceID = CreateVertexInstanceForContourVertex( VertexToInsert, PolygonID );
+		Polygon.PerimeterContour.VertexInstanceIDs.Insert( VertexInstanceID, InsertBeforeVertexNumber + InsertVertexIter );
+
+		// Connect polygon to vertex instance
+		FMeshVertexInstance& VertexInstance = VertexInstances[ VertexInstanceID.GetValue() ];
+		check( !VertexInstance.ConnectedPolygons.Contains( PolygonID ) );
+		VertexInstance.ConnectedPolygons.Add( PolygonID );
+	}
+
+	// Back up insert
+	{
+		FRemovePolygonPerimeterVerticesChangeInput RevertInput;
+
+		RevertInput.PolygonID = PolygonID;
+		RevertInput.FirstVertexNumberToRemove = InsertBeforeVertexNumber;
+		RevertInput.NumVerticesToRemove = VerticesToInsert.Num();
+		RevertInput.bDeleteOrphanedVertexInstances = false;
+
+		AddUndo( MakeUnique<FRemovePolygonPerimeterVerticesChange>( MoveTemp( RevertInput ) ) );
+	}
+
+	PolygonsPendingTriangulation.Add( PolygonID );
+	PolygonsPendingNewTangentBasis.Add( PolygonID );	// @todo mesheditor: and other vertex connected polygons too?
+
+	UE_LOG( LogEditableMesh, Verbose, TEXT( "* InsertPolygonPerimeterVertices finished" ) );
+}
+
+
+void UEditableMesh::RemovePolygonPerimeterVertices( const FPolygonID PolygonID, const int32 FirstVertexNumberToRemove, const int32 NumVerticesToRemove, const bool bDeleteOrphanedVertexInstances )
+{
+	UE_LOG( LogEditableMesh, Verbose, TEXT( "RemovePolygonPerimeterVertices: PolygonID:%s FirstVertexNumberToRemove:%d NumVerticesToRemove:%d" ), *PolygonID.ToString(), FirstVertexNumberToRemove, NumVerticesToRemove );
+	FMeshPolygon& Polygon = Polygons[ PolygonID.GetValue() ];
+
+	// Back up 
+	{
+		FInsertPolygonPerimeterVerticesChangeInput RevertInput;
+
+		RevertInput.PolygonID = PolygonID;
+		RevertInput.InsertBeforeVertexNumber = FirstVertexNumberToRemove;
+
+		RevertInput.VerticesToInsert.SetNum( NumVerticesToRemove, false );
+		for( int32 VertexToRemoveIter = 0; VertexToRemoveIter < NumVerticesToRemove; ++VertexToRemoveIter )
+		{
+			FVertexAndAttributes& RevertVertexToInsert = RevertInput.VerticesToInsert[ VertexToRemoveIter ];
+			RevertVertexToInsert.VertexInstanceID = Polygon.PerimeterContour.VertexInstanceIDs[ FirstVertexNumberToRemove + VertexToRemoveIter ];
+			RevertVertexToInsert.VertexID = FVertexID::Invalid;
+		}
+
+		AddUndo( MakeUnique<FInsertPolygonPerimeterVerticesChange>( MoveTemp( RevertInput ) ) );
+	}
+
+	{
+		static TArray<FVertexInstanceID> OrphanedVertexInstanceIDs;
+		OrphanedVertexInstanceIDs.Reset();
+
+		// Delete them backwards so it is done in the opposite order to Insert
+		for( int32 Index = NumVerticesToRemove - 1; Index >= 0; --Index )
+		{
+			const int32 VertexNumber = FirstVertexNumberToRemove + Index;
+			const FVertexInstanceID ContourVertexInstanceID = Polygon.PerimeterContour.VertexInstanceIDs[ VertexNumber ];
+
+			// Disconnect the polygon from the vertex instance
+			FMeshVertexInstance& VertexInstance = VertexInstances[ ContourVertexInstanceID.GetValue() ];
+			verify( VertexInstance.ConnectedPolygons.Remove( PolygonID ) == 1 );
+
+			// If the vertex instance is now orphaned, add it to the list
+			if( bDeleteOrphanedVertexInstances && VertexInstance.ConnectedPolygons.Num() == 0 )
+			{
+				OrphanedVertexInstanceIDs.Add( ContourVertexInstanceID );
+			}
+		}
+
+		Polygon.PerimeterContour.VertexInstanceIDs.RemoveAt( FirstVertexNumberToRemove, NumVerticesToRemove );
+
+		if( OrphanedVertexInstanceIDs.Num() > 0 )
+		{
+			DeleteVertexInstances( OrphanedVertexInstanceIDs, false );
+		}
+	}
+
+	PolygonsPendingTriangulation.Add( PolygonID );
+	PolygonsPendingNewTangentBasis.Add( PolygonID );	// @todo mesheditor: and other vertex connected polygons too?
+
+	UE_LOG( LogEditableMesh, Verbose, TEXT( "* RemovePolygonPerimeterVertices finished" ) );
+}
+
+
+int32 UEditableMesh::FindPolygonPerimeterVertexNumberForVertex( const FPolygonID PolygonID, const FVertexID VertexID ) const
+{
+	int32 FoundPolygonVertexNumber = INDEX_NONE;
+
+	const int32 PolygonPerimeterVertexCount = this->GetPolygonPerimeterVertexCount( PolygonID );
+	for( int32 PolygonVertexNumber = 0; PolygonVertexNumber < PolygonPerimeterVertexCount; ++PolygonVertexNumber )
+	{
+		if( VertexID == this->GetPolygonPerimeterVertex( PolygonID, PolygonVertexNumber ) )
+		{
+			FoundPolygonVertexNumber = PolygonVertexNumber;
+			break;
+		}
+	}
+
+	return FoundPolygonVertexNumber;
+}
+
+
+int32 UEditableMesh::FindPolygonHoleVertexNumberForVertex( const FPolygonID PolygonID, const int32 HoleNumber, const FVertexID VertexID ) const
+{
+	int32 FoundPolygonVertexNumber = INDEX_NONE;
+
+	const int32 PolygonHoleVertexCount = this->GetPolygonHoleVertexCount( PolygonID, HoleNumber );
+	for( int32 PolygonVertexNumber = 0; PolygonVertexNumber < PolygonHoleVertexCount; ++PolygonVertexNumber )
+	{
+		if( VertexID == this->GetPolygonHoleVertex( PolygonID, HoleNumber, PolygonVertexNumber ) )
+		{
+			FoundPolygonVertexNumber = PolygonVertexNumber;
+			break;
+		}
+	}
+
+	return FoundPolygonVertexNumber;
+}
+
+
+int32 UEditableMesh::FindPolygonPerimeterEdgeNumberForVertices( const FPolygonID PolygonID, const FVertexID EdgeVertexID0, const FVertexID EdgeVertexID1 ) const
+{
+	// @todo mesheditor: surely quicker just to iterate round the perimeter, checking pairs of adjacent vertices against the passed IDs (in both orders)?
+
+	int32 FoundPolygonEdgeNumber = INDEX_NONE;
+
+	static TArray<FEdgeID> EdgeIDs;
+	GetPolygonPerimeterEdges( PolygonID, /* Out */ EdgeIDs );
+
+	for( int32 PolygonEdgeNumber = 0; PolygonEdgeNumber < EdgeIDs.Num(); ++PolygonEdgeNumber )
+	{
+		const FEdgeID EdgeID = EdgeIDs[ PolygonEdgeNumber ];
+
+		FVertexID TestEdgeVertexIDs[ 2 ];
+		GetEdgeVertices( EdgeID, /* Out */ TestEdgeVertexIDs[0], /* Out */ TestEdgeVertexIDs[1] );
+
+		if( ( TestEdgeVertexIDs[ 0 ] == EdgeVertexID0 && TestEdgeVertexIDs[ 1 ] == EdgeVertexID1 ) ||
+			( TestEdgeVertexIDs[ 1 ] == EdgeVertexID0 && TestEdgeVertexIDs[ 0 ] == EdgeVertexID1 ) )
+		{
+			FoundPolygonEdgeNumber = PolygonEdgeNumber;
+			break;
+		}
+	}
+
+	return FoundPolygonEdgeNumber;
+}
+
+
+int32 UEditableMesh::FindPolygonHoleEdgeNumberForVertices( const FPolygonID PolygonID, const int32 HoleNumber, const FVertexID EdgeVertexID0, const FVertexID EdgeVertexID1 ) const
+{
+	// @todo mesheditor: surely quicker just to iterate round the perimeter, checking pairs of adjacent vertices against the passed IDs (in both orders)?
+
+	int32 FoundPolygonEdgeNumber = INDEX_NONE;
+
+	static TArray<FEdgeID> EdgeIDs;
+	GetPolygonHoleEdges( PolygonID, HoleNumber, /* Out */ EdgeIDs );
+
+	for( int32 PolygonEdgeNumber = 0; PolygonEdgeNumber < EdgeIDs.Num(); ++PolygonEdgeNumber )
+	{
+		const FEdgeID EdgeID = EdgeIDs[ PolygonEdgeNumber ];
+
+		FVertexID TestEdgeVertexIDs[ 2 ];
+		GetEdgeVertices( EdgeID, /* Out */ TestEdgeVertexIDs[0], /* Out */ TestEdgeVertexIDs[1] );
+
+		if( ( TestEdgeVertexIDs[ 0 ] == EdgeVertexID0 && TestEdgeVertexIDs[ 1 ] == EdgeVertexID1 ) ||
+			( TestEdgeVertexIDs[ 1 ] == EdgeVertexID0 && TestEdgeVertexIDs[ 0 ] == EdgeVertexID1 ) )
+		{
+			FoundPolygonEdgeNumber = PolygonEdgeNumber;
+			break;
+		}
+	}
+
+	return FoundPolygonEdgeNumber;
+}
+
+
+void UEditableMesh::FlipPolygons( const TArray<FPolygonID>& PolygonIDs )
+{
+	static TArray<FPolygonToCreate> PolygonsToCreate;
+	PolygonsToCreate.Reset();
+	for( int32 PolygonNumber = 0; PolygonNumber < PolygonIDs.Num(); ++PolygonNumber )
+	{
+		const FPolygonID OriginalPolygonID = PolygonIDs[ PolygonNumber ];
+		const FPolygonGroupID OriginalPolygonGroupID = GetGroupForPolygon( OriginalPolygonID );
+
+		FPolygonToCreate& PolygonToCreate = *new( PolygonsToCreate ) FPolygonToCreate();
+
+		PolygonToCreate.PolygonGroupID = OriginalPolygonGroupID;
+		
+		// Keep the original polygon ID.  No reason not to.
+		PolygonToCreate.OriginalPolygonID = OriginalPolygonID;
+
+		// Iterate backwards to add the vertices for the polygon, because we're flipping it over.
+		const int32 PerimeterVertexCount = GetPolygonPerimeterVertexCount( OriginalPolygonID );
+		PolygonToCreate.PerimeterVertices.Reserve( PerimeterVertexCount );
+		for( int32 VertexNumber = PerimeterVertexCount - 1; VertexNumber >= 0; --VertexNumber )
+		{
+			FVertexAndAttributes& PerimeterVertex = *new( PolygonToCreate.PerimeterVertices ) FVertexAndAttributes();
+
+			PerimeterVertex.VertexID = GetPolygonPerimeterVertex( OriginalPolygonID, VertexNumber );
+
+			// Save off information about the polygons, so we can re-create them with reverse winding.
+			TArray<FMeshElementAttributeData>& PerimeterVertexAttributeList = PerimeterVertex.PolygonVertexAttributes.Attributes;
+			for( const FName AttributeName : UEditableMesh::GetValidVertexInstanceAttributes() )
+			{
+				const int32 MaxAttributeIndex = GetMaxAttributeIndex( AttributeName );
+				for( int32 AttributeIndex = 0; AttributeIndex < MaxAttributeIndex; ++AttributeIndex )
+				{
+					FMeshElementAttributeData& PolygonVertexAttribute = *new( PerimeterVertexAttributeList ) FMeshElementAttributeData(
+						AttributeName,
+						AttributeIndex,
+						GetPolygonPerimeterVertexAttribute( OriginalPolygonID, VertexNumber, AttributeName, AttributeIndex ) );
+				}
+			}
+		}
+
+		// @todo mesheditor hole: Perserve holes in this polygon!  (Vertex IDs and Attributes!) Should these be flipped too?  Probably.
+	}
+
+	// Delete all of the polygons
+	{
+		const bool bDeleteOrphanedEdges = false;
+		const bool bDeleteOrphanedVertices = false;
+		const bool bDeleteOrphanedVertexInstances = false;
+		const bool bDeleteEmptyPolygonGroups = false;
+		DeletePolygons( PolygonIDs, bDeleteOrphanedEdges, bDeleteOrphanedVertices, bDeleteOrphanedVertexInstances, bDeleteEmptyPolygonGroups );
+	}
+
+	// Create the new polygons.  They're just like the polygons we deleted, except with reversed winding.
+	{
+		static TArray<FPolygonID> NewPolygonIDs;
+		static TArray<FEdgeID> NewEdgeIDs;
+		CreatePolygons( PolygonsToCreate, NewPolygonIDs, NewEdgeIDs );
+	}
+}
+
+
+void UEditableMesh::TriangulatePolygons( const TArray<FPolygonID>& PolygonIDs, TArray<FPolygonID>& OutNewTrianglePolygons )
+{
+	OutNewTrianglePolygons.Reset();
+
+	static TArray<FPolygonToCreate> PolygonsToCreate;
+	PolygonsToCreate.Reset();
+
+	static TArray<FPolygonID> PolygonsToDelete;
+	PolygonsToDelete.Reset();
+
+	for( const FPolygonID PolygonID : PolygonIDs )
+	{
+		// Skip right over polygons with fewer than four vertices
+		const int32 PerimeterVertexCount = GetPolygonPerimeterVertexCount( PolygonID );
+		if( PerimeterVertexCount > 3 )
+		{
+			// We'll be replacing this polygon with it's triangulated counterpart polygons
+			PolygonsToDelete.Add( PolygonID );
+
+			// Figure out the triangulation for this polygon
+			static TArray<FMeshTriangle> Triangles;
+			ComputePolygonTriangulation( PolygonID, /* Out */ Triangles );
+
+			// Build polygons for each of the triangles that made up the original
+			{
+				for( const FMeshTriangle& Triangle : Triangles )
+				{
+					PolygonsToCreate.Emplace();
+					FPolygonToCreate& PolygonToCreate = PolygonsToCreate.Last();
+					
+					PolygonToCreate.PolygonGroupID = GetGroupForPolygon( PolygonID );
+					PolygonToCreate.PolygonEdgeHardness = EPolygonEdgeHardness::NewEdgesSoft;
+
+					for( int32 TriangleVertexNumber = 0; TriangleVertexNumber < 3; ++TriangleVertexNumber )
+					{
+						PolygonToCreate.PerimeterVertices.Emplace();
+						FVertexAndAttributes& PerimeterVertex = PolygonToCreate.PerimeterVertices.Last();
+
+						PerimeterVertex.VertexInstanceID = Triangle.GetVertexInstanceID( TriangleVertexNumber );
+						PerimeterVertex.VertexID = FVertexID::Invalid;
+					}
+				}
+			}
+		}
+	}
+
+	// Delete the original polygons, but don't erase any orphaned edges or vertices, because we're about to put in
+	// triangles to replace those polygons.  Also, we won't touch polygons that we didn't have to triangulate!
+	{
+		const bool bDeleteOrphanEdges = false;
+		const bool bDeleteOrphanVertices = false;
+		const bool bDeleteOrphanVertexInstances = false;
+		const bool bDeleteEmptyPolygonGroups = false;
+		DeletePolygons( PolygonsToDelete, bDeleteOrphanEdges, bDeleteOrphanVertices, bDeleteOrphanVertexInstances, bDeleteEmptyPolygonGroups );
+	}
+
+	// Create the new polygons.  One for each triangle.  Note that new edges will be created here too on the inside of
+	// the original polygon to border the triangles.
+	{
+		static TArray<FEdgeID> NewEdgeIDs;
+		CreatePolygons( PolygonsToCreate, /* Out */ OutNewTrianglePolygons, /* Out */ NewEdgeIDs );
+	}
+}
+
+
+void UEditableMesh::AssignMaterialToPolygons( const TArray<FPolygonID>& PolygonIDs, UMaterialInterface* Material, TArray<FPolygonID>& NewPolygonIDs )
+{
+	// @todo mesheditor: create more generic method to assign a polygon to a new polygon group
+
+	// We need to move polygons from one polygon group to another. First find the new polygon group id to move them to.
+	const bool bCreateNewPolygonGroupIfNotFound = true;
+	const FPolygonGroupID NewPolygonGroupID = GetPolygonGroupIDFromMaterial( Material, bCreateNewPolygonGroupIfNotFound );
+	check( NewPolygonGroupID != FPolygonGroupID::Invalid );
+
+	// Make an array of FPolygonToCreate, based on the ones we wish to move. Everything is the same, other than the PolygonGroupID
+	static TArray<FPolygonToCreate> PolygonsToCreate;
+	PolygonsToCreate.Reset( PolygonIDs.Num() );
+
+	for( int32 PolygonNumber = 0; PolygonNumber < PolygonIDs.Num(); ++PolygonNumber )
+	{
+		const FPolygonID OriginalPolygonID = PolygonIDs[ PolygonNumber ];
+		const FMeshPolygon& Polygon = Polygons[ OriginalPolygonID.GetValue() ];
+
+		FPolygonToCreate& PolygonToCreate = *new( PolygonsToCreate ) FPolygonToCreate();
+
+		PolygonToCreate.PolygonGroupID = NewPolygonGroupID;
+		PolygonToCreate.OriginalPolygonID = FPolygonID::Invalid;
+
+		const int32 PerimeterVertexCount = GetPolygonPerimeterVertexCount( OriginalPolygonID );
+		PolygonToCreate.PerimeterVertices.Reserve( PerimeterVertexCount );
+
+		for( int32 VertexNumber = 0; VertexNumber < PerimeterVertexCount; ++VertexNumber )
+		{
+			FVertexAndAttributes& PerimeterVertex = *new( PolygonToCreate.PerimeterVertices ) FVertexAndAttributes();
+
+			PerimeterVertex.VertexID = GetPolygonPerimeterVertex( OriginalPolygonID, VertexNumber );
+
+			TArray<FMeshElementAttributeData>& PerimeterVertexAttributeList = PerimeterVertex.PolygonVertexAttributes.Attributes;
+
+			for( const FName AttributeName : UEditableMesh::GetValidVertexInstanceAttributes() )
+			{
+				const int32 MaxAttributeIndex = GetMaxAttributeIndex( AttributeName );
+				for( int32 AttributeIndex = 0; AttributeIndex < MaxAttributeIndex; ++AttributeIndex )
+				{
+					FMeshElementAttributeData& PolygonVertexAttribute = *new( PerimeterVertexAttributeList ) FMeshElementAttributeData(
+						AttributeName,
+						AttributeIndex,
+						GetPolygonPerimeterVertexAttribute( OriginalPolygonID, VertexNumber, AttributeName, AttributeIndex ) );
+				}
+			}
+		}
+
+		// @todo mesheditor: Handle holes?
+	}
+
+	// Create the new polygons in the new PolygonGroup.
+	{
+		static TArray<FEdgeID> NewEdgeIDs;
+		CreatePolygons( PolygonsToCreate, NewPolygonIDs, NewEdgeIDs );
+	}
+
+	// Delete the polygons from the old PolygonGroup
+	{
+		const bool bDeleteOrphanedEdges = false;
+		const bool bDeleteOrphanedVertices = false;
+		const bool bDeleteOrphanedVertexInstances = false;
+		const bool bDeleteEmptyPolygonGroups = true;
+		DeletePolygons( PolygonIDs, bDeleteOrphanedEdges, bDeleteOrphanedVertices, bDeleteOrphanedVertexInstances, bDeleteEmptyPolygonGroups );
+	}
+}
+
+
+FPolygonGroupID UEditableMesh::GetPolygonGroupIDFromMaterial( UMaterialInterface* Material, bool bCreateNewSectionIfNotFound )
+{
+	// Iterate through the sections sparse array looking for an entry whose material index matches.
+	for( TSparseArray<FMeshPolygonGroup>::TConstIterator It( PolygonGroups ); It; ++It )
+	{
+		const FMeshPolygonGroup& PolygonGroup = *It;
+
+		if( PolygonGroup.Material == Material )
+		{
+			return FPolygonGroupID( It.GetIndex() );
+		}
+	}
+
+	// If we got here, the material index does not yet have a matching section.
+	if( bCreateNewSectionIfNotFound )
+	{
+		static TArray<FPolygonGroupToCreate> PolygonGroupsToCreate;
+		PolygonGroupsToCreate.Reset( 1 );
+		PolygonGroupsToCreate.Emplace();
+		FPolygonGroupToCreate& PolygonGroupToCreate = PolygonGroupsToCreate[ 0 ];
+
+		PolygonGroupToCreate.Material = Material;
+		PolygonGroupToCreate.bEnableCollision = true;
+		PolygonGroupToCreate.bCastShadow = true;
+
+		static TArray<FPolygonGroupID> PolygonGroupIDs;
+		CreatePolygonGroups( PolygonGroupsToCreate, PolygonGroupIDs );
+		return PolygonGroupIDs[ 0 ];
+	}
+
+	return FPolygonGroupID::Invalid;
+}
+
+
+void UEditableMesh::WeldVertices( const TArray<FVertexID>& VertexIDsToWeld, FVertexID& OutNewVertexID )
+{
+	OutNewVertexID = FVertexID::Invalid;
+
+	// This function takes a list of perimeter vertices and a list of vertices to be welded as input.
+	// It returns a tuple stating whether the result is valid, and the [first, last) range of vertices to be welded.
+	// (It will be invalid if there is more than one contiguous run of vertices to weld.)
+	auto GetPerimeterVertexRangeToWeld = []( const TArray<FVertexID>& PolygonVertexIDs, const TArray<FVertexID>& VertexIDsToWeld )
+	{
+		bool bValid = true;
+		int32 StartIndex = INDEX_NONE;
+		int32 EndIndex = INDEX_NONE;
+
+		const int32 NumPolygonVertices = PolygonVertexIDs.Num();
+		bool bPrevVertexNeedsWelding = VertexIDsToWeld.Contains( PolygonVertexIDs[ NumPolygonVertices - 1 ] );
+		for( int32 Index = 0; Index < NumPolygonVertices; ++Index )
+		{
+			const bool bThisVertexNeedsWelding = VertexIDsToWeld.Contains( PolygonVertexIDs[ Index ] );
+			if( !bPrevVertexNeedsWelding && bThisVertexNeedsWelding )
+			{
+				// Transition from 'doesn't need welding' to 'needs welding'
+				if( StartIndex == INDEX_NONE )
+				{
+					StartIndex = Index;
+				}
+				else
+				{
+					// If this is not the first time we've seen this transition, there is more than one contiguous run of vertices
+					// which need welding, which is not allowed.
+					bValid = false;
+				}
+			}
+
+			if( bPrevVertexNeedsWelding && !bThisVertexNeedsWelding )
+			{
+				// Transition from 'needs welding' to 'doesn't need welding'
+				if( EndIndex == INDEX_NONE )
+				{
+					EndIndex = Index;
+				}
+				else
+				{
+					bValid = false;
+				}
+			}
+
+			bPrevVertexNeedsWelding = bThisVertexNeedsWelding;
+		}
+
+		// If the indices are not set, either there were no vertices to weld, or they were all to be welded.
+		// In the latter case, initialize the full vertex range.
+		if( StartIndex == INDEX_NONE && EndIndex == INDEX_NONE && bPrevVertexNeedsWelding )
+		{
+			StartIndex = 0;
+			EndIndex = NumPolygonVertices;
+		}
+
+		// Get the size of the range.
+		// The array is circular, so it's possible for the EndIndex to be smaller than the start index (and compensate for that accordingly)
+		const int32 RangeSize = ( EndIndex - StartIndex ) + ( ( EndIndex < StartIndex ) ? NumPolygonVertices : 0 );
+
+		// If, after welding perimeter vertices, we have fewer than three vertices left, this poly just disappears.
+		// (+ 1 below for the new vertex which replaces the welded range)
+		bool bWouldBeDegenerate = ( NumPolygonVertices - RangeSize + 1 < 3 );
+
+		return MakeTuple( bValid, bWouldBeDegenerate, StartIndex, EndIndex );
+	};
+
+	// Build a list of all polygons which contain at least one of the vertices to be welded
+	static TArray<FPolygonID> AllConnectedPolygonIDs;
+	{
+		AllConnectedPolygonIDs.Reset();
+
+		for( FVertexID VertexID : VertexIDsToWeld )
+		{
+			static TArray<FPolygonID> ConnectedPolygonIDs;
+			GetVertexConnectedPolygons( VertexID, ConnectedPolygonIDs );
+
+			for( FPolygonID PolygonID : ConnectedPolygonIDs )
+			{
+				AllConnectedPolygonIDs.AddUnique( PolygonID );
+			}
+		}
+	}
+
+	// Check whether the operation is valid. We can't weld vertices if there are any polygons which have non-contiguous
+	// vertices on the perimeter which are marked to be welded
+	bool bNeedToCreateWeldedVertex = false;
+	for( FPolygonID ConnectedPolygonID : AllConnectedPolygonIDs )
+	{
+		static TArray<FVertexID> PolygonVertexIDs;
+		GetPolygonPerimeterVertices( ConnectedPolygonID, PolygonVertexIDs );
+
+		auto VertexRangeToWeld = GetPerimeterVertexRangeToWeld( PolygonVertexIDs, VertexIDsToWeld );
+		const bool bIsValid = VertexRangeToWeld.Get<0>();
+		const bool bWouldBeDegenerate = VertexRangeToWeld.Get<1>();
+
+		// If the resulting poly is valid (has 3 or more verts), we know we need to create the welded vertex
+		if( !bWouldBeDegenerate )
+		{
+			bNeedToCreateWeldedVertex = true;
+		}
+
+		// If the result is invalid (because it would cause a poly to be welded in more than one place on its perimeter), abort now
+		if( !bIsValid )
+		{
+			// Return with the NewVertexID set to Invalid
+			return;
+		}
+	}
+
+	if( !bNeedToCreateWeldedVertex )
+	{
+		// For now, abort if we don't need to create a welded vertex.
+		// This generally implies that all (or a disconnected subset) of the mesh is about to disappear,
+		// which arguably is not something we would want to do like this anyway.
+		return;
+	}
+
+	// Create new welded vertex
+	static TArray<FVertexID> NewVertices;
+	{
+		static TArray<FVertexToCreate> VerticesToCreate;
+		VerticesToCreate.SetNum( 1 );
+		FVertexToCreate& VertexToCreate = VerticesToCreate.Last();
+
+		// The vertex which is created will be at the position of the last vertex in the array of vertices to weld.
+		// @todo mesheditor: maybe specify a position to the method instead of using one of the existing vertices?
+		VertexToCreate.VertexAttributes.Attributes.Emplace(
+			UEditableMeshAttribute::VertexPosition(),
+			0,
+			GetVertexAttribute( VertexIDsToWeld.Last(), UEditableMeshAttribute::VertexPosition(), 0 )
+		);
+
+		// @todo mesheditor: vertex corner sharpness too?
+
+		CreateVertices( VerticesToCreate, NewVertices );
+	}
+
+
+	static TArray<FPolygonToCreate> PolygonsToCreate;
+	PolygonsToCreate.Reset( AllConnectedPolygonIDs.Num() );
+
+	static TArray<FAttributesForEdge> AttributesForEdges;
+	AttributesForEdges.Reset();
+
+	// Now for each polygon, merge runs of consecutive vertices
+	for( FPolygonID ConnectedPolygonID : AllConnectedPolygonIDs )
+	{
+		const int32 NumPerimeterVertices = GetPolygonPerimeterVertexCount( ConnectedPolygonID );
+
+		// Get perimeter vertices and edges for this polygon
+		static TArray<FVertexID> PolygonVertexIDs;
+		static TArray<FEdgeID> PolygonEdgeIDs;
+		GetPolygonPerimeterVertices( ConnectedPolygonID, PolygonVertexIDs );
+		GetPolygonPerimeterEdges( ConnectedPolygonID, PolygonEdgeIDs );
+
+		// Get the index range of perimeter vertices to be welded.
+		// This should definitely be valid, as any invalid welded poly will have caused early exit, above.
+		auto VertexRangeToWeld = GetPerimeterVertexRangeToWeld( PolygonVertexIDs, VertexIDsToWeld );
+		const bool bIsValid = VertexRangeToWeld.Get<0>();
+		const bool bWouldBeDegenerate = VertexRangeToWeld.Get<1>();
+		const int32 StartIndex = VertexRangeToWeld.Get<2>();
+		const int32 EndIndex = VertexRangeToWeld.Get<3>();
+		check( bIsValid );
+
+		if( bWouldBeDegenerate )
+		{
+			continue;
+		}
+
+		// Prepare to create a new polygon
+		PolygonsToCreate.Emplace();
+		FPolygonToCreate& PolygonToCreate = PolygonsToCreate.Last();
+
+		PolygonToCreate.PolygonGroupID = GetGroupForPolygon( ConnectedPolygonID );
+
+		// Iterate through perimeter vertices starting at index 0.
+		// We skip through the run of welded vertices, replacing them with a single welded vertex.
+		// We need to check whether we are starting in the middle of a run (if EndIndex < StartIndex).
+		bool bInsideWeldedRange = ( EndIndex < StartIndex );
+		for( int32 Index = 0; Index < NumPerimeterVertices; ++Index )
+		{
+			if( bInsideWeldedRange )
+			{
+				if( Index == EndIndex )
+				{
+					// EndIndex is range exclusive, so we now need to process this vertex.
+					bInsideWeldedRange = false;
+				}
+				else
+				{
+					// Otherwise still inside the welded range; skip the remaining vertices in the range.
+					continue;
+				}
+			}
+
+			// Add new perimeter vertex in the polygon to create
+			PolygonToCreate.PerimeterVertices.Emplace();
+			FVertexAndAttributes& VertexAndAttributes = PolygonToCreate.PerimeterVertices.Last();
+
+			if( Index == StartIndex )
+			{
+				// If this is the first vertex in the run of vertices to weld, replace the ID with the newly created welded vertex
+				VertexAndAttributes.VertexID = NewVertices[ 0 ];
+				bInsideWeldedRange = true;
+			}
+			else
+			{
+				// Otherwise use the original Vertex ID
+				VertexAndAttributes.VertexID = PolygonVertexIDs[ Index ];
+			}
+
+			// Copy the polygon vertex attributes over
+			for( int32 TextureCoordinateIndex = 0; TextureCoordinateIndex < TextureCoordinateCount; ++TextureCoordinateIndex )
+			{
+				VertexAndAttributes.PolygonVertexAttributes.Attributes.Add( FMeshElementAttributeData(
+					UEditableMeshAttribute::VertexTextureCoordinate(),
+					TextureCoordinateIndex,
+					GetPolygonPerimeterVertexAttribute( ConnectedPolygonID, Index, UEditableMeshAttribute::VertexTextureCoordinate(), TextureCoordinateIndex ) )
+				);
+			}
+
+			VertexAndAttributes.PolygonVertexAttributes.Attributes.Add( FMeshElementAttributeData(
+				UEditableMeshAttribute::VertexColor(),
+				0,
+				GetPolygonPerimeterVertexAttribute( ConnectedPolygonID, Index, UEditableMeshAttribute::VertexColor(), 0 ) )
+			);
+
+			// Prepare to assign the old edge's attributes to the new edge.
+			// We build up an array of edge attributes to set, in perimeter vertex order for each polygon.
+			AttributesForEdges.Emplace();
+			FAttributesForEdge& AttributesForEdge = AttributesForEdges.Last();
+
+			AttributesForEdge.EdgeAttributes.Attributes.Add( FMeshElementAttributeData(
+				UEditableMeshAttribute::EdgeIsHard(),
+				0,
+				GetEdgeAttribute( PolygonEdgeIDs[ Index ], UEditableMeshAttribute::EdgeIsHard(), 0 ) )
+			);
+			AttributesForEdge.EdgeAttributes.Attributes.Add( FMeshElementAttributeData( 
+				UEditableMeshAttribute::EdgeCreaseSharpness(), 
+				0, 
+				GetEdgeAttribute( PolygonEdgeIDs[ Index ], UEditableMeshAttribute::EdgeCreaseSharpness(), 0 ) )
+			);
+		}
+	}
+
+	// Create polygons
+	static TArray<FPolygonID> NewPolygonIDs;
+	{
+		static TArray<FEdgeID> NewEdgeIDs;
+		CreatePolygons( PolygonsToCreate, NewPolygonIDs, NewEdgeIDs );
+	}
+
+	// Set new edge attributes.
+	// Now that we have a list of newly created polygon IDs, we need to go through the attributes for edge list, filling in the new Edge ID.
+	// This relies on the fact that the NewPolygonIDs array lists the polygons in the same order as they were defined in PolygonsToCreate,
+	// and that the edges are strictly ordered from perimeter vertex 0.
+	{
+		int32 AttributesForEdgeIndex = 0;
+		for( FPolygonID NewPolygonID : NewPolygonIDs )
+		{
+			static TArray<FEdgeID> NewPolygonEdgeIDs;
+			GetPolygonPerimeterEdges( NewPolygonID, NewPolygonEdgeIDs );
+
+			for( FEdgeID NewPolygonEdgeID : NewPolygonEdgeIDs )
+			{
+				AttributesForEdges[ AttributesForEdgeIndex ].EdgeID = NewPolygonEdgeID;
+				AttributesForEdgeIndex++;
+			}
+		}
+		check( AttributesForEdgeIndex == AttributesForEdges.Num() );
+		SetEdgesAttributes( AttributesForEdges );
+	}
+
+	// Delete old polygons, removing any orphaned edges and vertices at the same time
+	{
+		const bool bDeleteOrphanedEdges = true;
+		const bool bDeleteOrphanedVertices = true;
+		const bool bDeleteOrphanedVertexInstances = true;
+		const bool bDeleteEmptyPolygonGroups = false;
+		DeletePolygons( AllConnectedPolygonIDs, bDeleteOrphanedEdges, bDeleteOrphanedVertices, bDeleteOrphanedVertexInstances, bDeleteEmptyPolygonGroups );
+	}
+}
+
+
+void UEditableMesh::TessellatePolygons( const TArray<FPolygonID>& PolygonIDs, const ETriangleTessellationMode TriangleTessellationMode, TArray<FPolygonID>& OutNewPolygonIDs )
+{
+	OutNewPolygonIDs.Reset();
+
+	//
+	// Simple tessellation algorithm:
+	//
+	//   - Triangles will be split into either three or four triangles depending on the 'Mode' argument.
+	//			-> ThreeTriangles: Connect each vertex to a new center vertex, forming three triangles
+	//			-> FourTriangles: Split each edge and create a center polygon that connects those new vertices, then three additional polygons for each original corner
+	//
+	//   - Everything else will be split into quads by creating a new vertex in the center, then adding a new vertex to 
+	//       each original perimeter edge and connecting each original vertex to it's new neighbors and the center
+	//
+	// NOTE: Concave polygons will yield bad results
+	//
+
+	// Create a new vertex in the center of each incoming polygon
+	static TArray<FVertexID> PolygonCenterVertices;
+	PolygonCenterVertices.Reset();
+	{
+		static TArray<FVertexToCreate> VerticesToCreate;
+		VerticesToCreate.Reset();
+		for( const FPolygonID PolygonID : PolygonIDs )
+		{
+			const int32 PerimeterEdgeCount = GetPolygonPerimeterEdgeCount( PolygonID );
+			if( TriangleTessellationMode == ETriangleTessellationMode::ThreeTriangles || PerimeterEdgeCount > 3 )
+			{
+				// Find the center of this polygon
+				const FVector PolygonCenter = ComputePolygonCenter( PolygonID );
+
+				FVertexToCreate& VertexToCreate = *new( VerticesToCreate ) FVertexToCreate();
+				VertexToCreate.VertexAttributes.Attributes.Add( FMeshElementAttributeData( UEditableMeshAttribute::VertexPosition(), 0, PolygonCenter ) );
+			}
+		}
+
+		this->CreateVertices( VerticesToCreate, /* Out */ PolygonCenterVertices );
+	}
+
+
+	// Split all of the edges of the original polygons (except triangles).  Remember, some edges may be shared between
+	// the incoming polygons so we'll keep track of that and make sure not to split them again.  
+	{
+		static TSet<FEdgeID> EdgesToSplit;
+		EdgesToSplit.Reset();
+
+		for( int32 PolygonNumber = 0; PolygonNumber < PolygonIDs.Num(); ++PolygonNumber )
+		{
+			const FPolygonID PolygonID = PolygonIDs[ PolygonNumber ];
+
+			static TArray<FEdgeID> PerimeterEdgeIDs;
+			GetPolygonPerimeterEdges( PolygonID, /* Out */ PerimeterEdgeIDs );
+			
+			if( TriangleTessellationMode == ETriangleTessellationMode::FourTriangles || PerimeterEdgeIDs.Num() > 3 )
+			{
+				for( const FEdgeID PerimeterEdgeID : PerimeterEdgeIDs )
+				{
+					EdgesToSplit.Add( PerimeterEdgeID );
+				}
+			}
+		}
+
+		for( const FEdgeID EdgeID : EdgesToSplit )
+		{
+			// Split the edge
+			static TArray<float> Splits;
+			Splits.SetNumUninitialized( 1 );
+			Splits[ 0 ] = 0.5f;
+
+			static TArray<FVertexID> NewVertexIDsFromSplit;
+			this->SplitEdge( EdgeID, Splits, /* Out */ NewVertexIDsFromSplit );
+			check( NewVertexIDsFromSplit.Num() == 1 );
+		}
+	}
+
+
+	// We'll now define the new polygons to be created.
+	static TArray<FPolygonToCreate> PolygonsToCreate;
+	PolygonsToCreate.Reset();
+
+	int32 PolygonWithNewCenterVertexNumber = 0;
+	for( int32 PolygonNumber = 0; PolygonNumber < PolygonIDs.Num(); ++PolygonNumber )
+	{
+		const FPolygonID PolygonID = PolygonIDs[ PolygonNumber ];
+		const FPolygonGroupID PolygonGroupID = GetGroupForPolygon( PolygonID );
+
+		const int32 PerimeterEdgeCount = GetPolygonPerimeterEdgeCount( PolygonID );
+
+		FVertexID PolygonCenterVertexID = FVertexID::Invalid;
+		if( TriangleTessellationMode == ETriangleTessellationMode::ThreeTriangles || PerimeterEdgeCount > 6 )
+		{
+			PolygonCenterVertexID = PolygonCenterVertices[ PolygonWithNewCenterVertexNumber++ ];
+		}
+
+		// Don't bother with triangles, because we'll simply connect the original three vertices to a new
+		// center position to tessellate those.
+		if( PerimeterEdgeCount > 6 )
+		{
+			static TArray<FVertexID> PerimeterVertexIDs;
+			GetPolygonPerimeterVertices( PolygonID, /* Out */ PerimeterVertexIDs );
+
+			const int32 PerimeterVertexCount = PerimeterEdgeCount;
+			const int32 OriginalPerimeterEdgeCount = PerimeterEdgeCount / 2;
+			for( int32 OriginalPerimeterEdgeNumber = 0; OriginalPerimeterEdgeNumber < OriginalPerimeterEdgeCount; ++OriginalPerimeterEdgeNumber )
+			{
+				const int32 CurrentVertexNumber = OriginalPerimeterEdgeNumber * 2;
+				const int32 PreviousVertexNumber = ( ( CurrentVertexNumber - 1 ) + PerimeterVertexCount ) % PerimeterVertexCount;
+				const int32 NextVertexNumber = ( CurrentVertexNumber + 1 ) % PerimeterVertexCount;
+
+				FPolygonToCreate& PolygonToCreate = *new( PolygonsToCreate ) FPolygonToCreate();
+				PolygonToCreate.PolygonGroupID = PolygonGroupID;
+
+				for( int32 QuadVertexNumber = 0; QuadVertexNumber < 4; ++QuadVertexNumber )
+				{
+					FVertexAndAttributes& VertexAndAttributes = *new( PolygonToCreate.PerimeterVertices ) FVertexAndAttributes();
+
+					int32 PerimeterVertexNumber = INDEX_NONE;
+					switch( QuadVertexNumber )
+					{
+						case 0:
+							PerimeterVertexNumber = PreviousVertexNumber;
+							break;
+
+						case 1:
+							PerimeterVertexNumber = CurrentVertexNumber;
+							break;
+
+						case 2:
+							PerimeterVertexNumber = NextVertexNumber;
+							break;
+
+						case 3:
+							PerimeterVertexNumber = INDEX_NONE;	// The center vertex!
+							break;
+
+						default:
+							check( 0 );
+					}
+
+
+					if( PerimeterVertexNumber == INDEX_NONE )
+					{
+						VertexAndAttributes.VertexID = PolygonCenterVertexID;
+
+						// Generate interpolated UVs and vertex colors for the new vertex in the center
+						{
+							const FVector CenterVertexPosition = GetVertexAttribute( PolygonCenterVertexID, UEditableMeshAttribute::VertexPosition(), 0 );
+
+							FMeshTriangle Triangle;
+							FVector TriangleVertexWeights;
+							if( ComputeBarycentricWeightForPointOnPolygon( PolygonID, CenterVertexPosition, /* Out */ Triangle, /* Out */ TriangleVertexWeights ) )
+							{
+								for( int32 TextureCoordinateIndex = 0; TextureCoordinateIndex < TextureCoordinateCount; ++TextureCoordinateIndex )
+								{
+									VertexAndAttributes.PolygonVertexAttributes.Attributes.Add( FMeshElementAttributeData(
+										UEditableMeshAttribute::VertexTextureCoordinate(),
+										TextureCoordinateIndex,
+										TriangleVertexWeights.X * GetVertexInstanceAttribute( Triangle.VertexInstanceID0, UEditableMeshAttribute::VertexTextureCoordinate(), TextureCoordinateIndex ) +
+										TriangleVertexWeights.Y * GetVertexInstanceAttribute( Triangle.VertexInstanceID1, UEditableMeshAttribute::VertexTextureCoordinate(), TextureCoordinateIndex ) +
+										TriangleVertexWeights.Z * GetVertexInstanceAttribute( Triangle.VertexInstanceID2, UEditableMeshAttribute::VertexTextureCoordinate(), TextureCoordinateIndex ) ) );
+								}
+
+								VertexAndAttributes.PolygonVertexAttributes.Attributes.Add( FMeshElementAttributeData(
+									UEditableMeshAttribute::VertexColor(),
+									0,
+									TriangleVertexWeights.X * GetVertexInstanceAttribute( Triangle.VertexInstanceID0, UEditableMeshAttribute::VertexColor(), 0 ) +
+									TriangleVertexWeights.Y * GetVertexInstanceAttribute( Triangle.VertexInstanceID1, UEditableMeshAttribute::VertexColor(), 0 ) +
+									TriangleVertexWeights.Z * GetVertexInstanceAttribute( Triangle.VertexInstanceID2, UEditableMeshAttribute::VertexColor(), 0 ) ) );
+							}
+						}
+					}
+					else
+					{
+						VertexAndAttributes.VertexID = PerimeterVertexIDs[ PerimeterVertexNumber ];
+
+						// Transfer over UVs and vertex colors from the original polygon that we're replacing
+						{
+							for( int32 TextureCoordinateIndex = 0; TextureCoordinateIndex < TextureCoordinateCount; ++TextureCoordinateIndex )
+							{
+								VertexAndAttributes.PolygonVertexAttributes.Attributes.Add( FMeshElementAttributeData(
+									UEditableMeshAttribute::VertexTextureCoordinate(),
+									TextureCoordinateIndex,
+									GetPolygonPerimeterVertexAttribute( PolygonID, PerimeterVertexNumber, UEditableMeshAttribute::VertexTextureCoordinate(), TextureCoordinateIndex ) ) );
+							}
+
+							VertexAndAttributes.PolygonVertexAttributes.Attributes.Add( FMeshElementAttributeData(
+								UEditableMeshAttribute::VertexColor(),
+								0,
+								GetPolygonPerimeterVertexAttribute( PolygonID, PerimeterVertexNumber, UEditableMeshAttribute::VertexColor(), 0 ) ) );
+						}
+					}
+				}
+			}
+		}
+		else
+		{
+			if( TriangleTessellationMode == ETriangleTessellationMode::ThreeTriangles )
+			{
+				// Define the three new triangles for the original tessellated triangle
+				for( int32 PerimeterEdgeNumber = 0; PerimeterEdgeNumber < 3; ++PerimeterEdgeNumber )
+				{
+					bool bIsEdgeWindingReversedForPolygon;
+					const FEdgeID EdgeID = GetPolygonPerimeterEdge( PolygonID, PerimeterEdgeNumber, /* Out */ bIsEdgeWindingReversedForPolygon );
+
+					FVertexID EdgeVertexID0, EdgeVertexID1;
+					GetEdgeVertices( EdgeID, /* Out */ EdgeVertexID0, /* Out */ EdgeVertexID1 );
+
+					FPolygonToCreate& PolygonToCreate = *new( PolygonsToCreate ) FPolygonToCreate();
+					PolygonToCreate.PolygonGroupID = PolygonGroupID;
+
+					for( int32 TriangleVertexNumber = 0; TriangleVertexNumber < 3; ++TriangleVertexNumber )
+					{
+						FVertexAndAttributes& VertexAndAttributes = *new( PolygonToCreate.PerimeterVertices ) FVertexAndAttributes();
+						switch( TriangleVertexNumber )
+						{
+							case 0:
+								VertexAndAttributes.VertexID = bIsEdgeWindingReversedForPolygon ? EdgeVertexID1 : EdgeVertexID0;
+								break;
+
+							case 1:
+								VertexAndAttributes.VertexID = PolygonCenterVertexID;
+								break;
+
+							case 2:
+								VertexAndAttributes.VertexID = bIsEdgeWindingReversedForPolygon ? EdgeVertexID0 : EdgeVertexID1;
+								break;
+
+							default:
+								check( 0 );
+						}
+
+						if( VertexAndAttributes.VertexID == PolygonCenterVertexID )
+						{
+							// Generate interpolated UVs and vertex colors for the new vertex in the center
+							{
+								const FVector CenterVertexPosition = GetVertexAttribute( PolygonCenterVertexID, UEditableMeshAttribute::VertexPosition(), 0 );
+
+								FMeshTriangle Triangle;
+								FVector TriangleVertexWeights;
+								if( ComputeBarycentricWeightForPointOnPolygon( PolygonID, CenterVertexPosition, /* Out */ Triangle, /* Out */ TriangleVertexWeights ) )
+								{
+									for( int32 TextureCoordinateIndex = 0; TextureCoordinateIndex < TextureCoordinateCount; ++TextureCoordinateIndex )
+									{
+										VertexAndAttributes.PolygonVertexAttributes.Attributes.Add( FMeshElementAttributeData(
+											UEditableMeshAttribute::VertexTextureCoordinate(),
+											TextureCoordinateIndex,
+											TriangleVertexWeights.X * GetVertexInstanceAttribute( Triangle.VertexInstanceID0, UEditableMeshAttribute::VertexTextureCoordinate(), TextureCoordinateIndex ) +
+											TriangleVertexWeights.Y * GetVertexInstanceAttribute( Triangle.VertexInstanceID1, UEditableMeshAttribute::VertexTextureCoordinate(), TextureCoordinateIndex ) +
+											TriangleVertexWeights.Z * GetVertexInstanceAttribute( Triangle.VertexInstanceID2, UEditableMeshAttribute::VertexTextureCoordinate(), TextureCoordinateIndex ) ) );
+									}
+
+									VertexAndAttributes.PolygonVertexAttributes.Attributes.Add( FMeshElementAttributeData(
+										UEditableMeshAttribute::VertexColor(),
+										0,
+										TriangleVertexWeights.X * GetVertexInstanceAttribute( Triangle.VertexInstanceID0, UEditableMeshAttribute::VertexColor(), 0 ) +
+										TriangleVertexWeights.Y * GetVertexInstanceAttribute( Triangle.VertexInstanceID1, UEditableMeshAttribute::VertexColor(), 0 ) +
+										TriangleVertexWeights.Z * GetVertexInstanceAttribute( Triangle.VertexInstanceID2, UEditableMeshAttribute::VertexColor(), 0 ) ) );
+								}
+							}
+						}
+						else
+						{
+							// Transfer over UVs and vertex colors from the original polygon that we're replacing
+							{
+								const int32 VertexNumber = FindPolygonPerimeterVertexNumberForVertex( PolygonID, VertexAndAttributes.VertexID );
+								check( VertexNumber != INDEX_NONE );
+
+								for( int32 TextureCoordinateIndex = 0; TextureCoordinateIndex < TextureCoordinateCount; ++TextureCoordinateIndex )
+								{
+									VertexAndAttributes.PolygonVertexAttributes.Attributes.Add( FMeshElementAttributeData(
+										UEditableMeshAttribute::VertexTextureCoordinate(),
+										TextureCoordinateIndex,
+										GetPolygonPerimeterVertexAttribute( PolygonID, VertexNumber, UEditableMeshAttribute::VertexTextureCoordinate(), TextureCoordinateIndex ) ) );
+								}
+
+								VertexAndAttributes.PolygonVertexAttributes.Attributes.Add( FMeshElementAttributeData(
+									UEditableMeshAttribute::VertexColor(),
+									0,
+									GetPolygonPerimeterVertexAttribute( PolygonID, VertexNumber, UEditableMeshAttribute::VertexColor(), 0 ) ) );
+							}
+						}
+					}
+				}
+			}
+			else if( ensure( TriangleTessellationMode == ETriangleTessellationMode::FourTriangles ) )
+			{
+				// Define the four new triangles for the original tessellated triangle.  One triangle will go in
+				// the center, connecting the three new vertices that we created between each original edge.  The
+				// other three triangles will go in the corners of the original triangle.
+
+				static TArray<FVertexID> PerimeterVertexIDs;
+				GetPolygonPerimeterVertices( PolygonID, PerimeterVertexIDs );
+				check( PerimeterVertexIDs.Num() == 6 );	// We split the triangle's 3 edges earlier, so we must have six edges now
+
+				// Define the new center triangle
+				{
+					FPolygonToCreate& PolygonToCreate = *new( PolygonsToCreate ) FPolygonToCreate();
+					PolygonToCreate.PolygonGroupID = PolygonGroupID;
+
+					for( int32 OriginalVertexNumber = 0; OriginalVertexNumber < 3; ++OriginalVertexNumber )
+					{
+						const int32 VertexNumber = ( OriginalVertexNumber * 2 + 1 ) % PerimeterVertexIDs.Num();
+
+						FVertexAndAttributes& VertexAndAttributes = *new( PolygonToCreate.PerimeterVertices ) FVertexAndAttributes();
+						VertexAndAttributes.VertexID = PerimeterVertexIDs[ VertexNumber ];
+
+						// Transfer over UVs and vertex colors from the original polygon that we're replacing
+						{
+							for( int32 TextureCoordinateIndex = 0; TextureCoordinateIndex < TextureCoordinateCount; ++TextureCoordinateIndex )
+							{
+								VertexAndAttributes.PolygonVertexAttributes.Attributes.Add( FMeshElementAttributeData(
+									UEditableMeshAttribute::VertexTextureCoordinate(),
+									TextureCoordinateIndex,
+									GetPolygonPerimeterVertexAttribute( PolygonID, VertexNumber, UEditableMeshAttribute::VertexTextureCoordinate(), TextureCoordinateIndex ) ) );
+							}
+
+							VertexAndAttributes.PolygonVertexAttributes.Attributes.Add( FMeshElementAttributeData(
+								UEditableMeshAttribute::VertexColor(),
+								0,
+								GetPolygonPerimeterVertexAttribute( PolygonID, VertexNumber, UEditableMeshAttribute::VertexColor(), 0 ) ) );
+						}
+					}
+				}
+
+				// Define the three corner triangles
+				for( int32 OriginalEdgeNumber = 0; OriginalEdgeNumber < 3; ++OriginalEdgeNumber )
+				{
+					const int32 CurrentVertexNumber = OriginalEdgeNumber * 2;
+					const int32 PreviousVertexNumber = ( ( CurrentVertexNumber - 1 ) + PerimeterVertexIDs.Num() ) % PerimeterVertexIDs.Num();
+					const int32 NextVertexNumber = ( CurrentVertexNumber + 1 ) % PerimeterVertexIDs.Num();
+
+					FPolygonToCreate& PolygonToCreate = *new( PolygonsToCreate ) FPolygonToCreate();
+					PolygonToCreate.PolygonGroupID = PolygonGroupID;
+
+					for( int32 TriangleVertexNumber = 0; TriangleVertexNumber < 3; ++TriangleVertexNumber )
+					{
+						int32 VertexNumber = INDEX_NONE;
+						switch( TriangleVertexNumber )
+						{
+							case 0:
+								VertexNumber = PreviousVertexNumber;
+								break;
+
+							case 1:
+								VertexNumber = CurrentVertexNumber;
+								break;
+
+							case 2:
+								VertexNumber = NextVertexNumber;
+								break;
+
+							default:
+								check( 0 );
+						}
+
+						FVertexAndAttributes& VertexAndAttributes = *new( PolygonToCreate.PerimeterVertices ) FVertexAndAttributes();
+						VertexAndAttributes.VertexID = PerimeterVertexIDs[ VertexNumber ];
+
+						{
+							// Transfer over UVs and vertex colors from the original polygon that we're replacing
+							{
+								for( int32 TextureCoordinateIndex = 0; TextureCoordinateIndex < TextureCoordinateCount; ++TextureCoordinateIndex )
+								{
+									VertexAndAttributes.PolygonVertexAttributes.Attributes.Add( FMeshElementAttributeData(
+										UEditableMeshAttribute::VertexTextureCoordinate(),
+										TextureCoordinateIndex,
+										GetPolygonPerimeterVertexAttribute( PolygonID, VertexNumber, UEditableMeshAttribute::VertexTextureCoordinate(), TextureCoordinateIndex ) ) );
+								}
+
+								VertexAndAttributes.PolygonVertexAttributes.Attributes.Add( FMeshElementAttributeData(
+									UEditableMeshAttribute::VertexColor(),
+									0,
+									GetPolygonPerimeterVertexAttribute( PolygonID, VertexNumber, UEditableMeshAttribute::VertexColor(), 0 ) ) );
+							}
+						}
+					}
+				}
+			}
+		}
+	}
+
+
+	// Delete the original polygons
+	{
+		const bool bDeleteOrphanedEdges = false;	// No need to delete orphans, because this function won't orphan anything
+		const bool bDeleteOrphanedVertices = false;
+		const bool bDeleteOrphanedVertexInstances = false;
+		const bool bDeleteEmptyPolygonGroups = false;
+		DeletePolygons( PolygonIDs, bDeleteOrphanedEdges, bDeleteOrphanedVertices, bDeleteOrphanedVertexInstances, bDeleteEmptyPolygonGroups );
+	}
+
+
+	// Create all of the new polygons for the tessellated representation of the original polygons
+	{
+		static TArray<FEdgeID> NewEdgeIDs;
+		CreatePolygons( PolygonsToCreate, /* Out */ OutNewPolygonIDs, /* Out */ NewEdgeIDs );
+	}
+}
+
+
+void UEditableMesh::SetTextureCoordinateCount( const int32 NumTexCoords )
+{
+	TextureCoordinateCount = FMath::Max( NumTexCoords, 0 );
+}
+
+
+void UEditableMesh::QuadrangulateMesh( TArray<FPolygonID>& OutNewPolygonIDs )
+{
+	// Iterate through all polygon groups in the mesh, quadrangulating each in turn
+	const int32 MaxPolygonGroupIndex = GetPolygonGroupArraySize();
+	for( int32 PolygonGroupIndex = 0; PolygonGroupIndex < MaxPolygonGroupIndex; ++PolygonGroupIndex )
+	{
+		const FPolygonGroupID PolygonGroupID( PolygonGroupIndex );
+		if( IsValidPolygonGroup( PolygonGroupID ) )
+		{
+			static TArray<FPolygonID> NewPolygonIDs;
+			QuadrangulatePolygonGroup( PolygonGroupID, NewPolygonIDs );
+			OutNewPolygonIDs.Append( NewPolygonIDs );
+		}
+	}
+}
+
+
+void UEditableMesh::QuadrangulatePolygonGroup( const FPolygonGroupID PolygonGroupID, TArray<FPolygonID>& OutNewPolygonIDs )
+{
+	// Uses the first two steps of the algorithm described by
+	// http://www.lirmm.fr/~beniere/ArticlesPersos/GRAPP10_Beniere_Final.pdf
+
+	// Tweakable parameters affecting how quadrangulate works
+	const float CosAngleThreshold = 0.984f;		// about 10 degrees
+	const bool bKeepHardEdges = true;
+	const bool bKeepTextureBorder = true;
+	const bool bKeepColorBorder = true;
+
+	OutNewPolygonIDs.Reset();
+
+	static TArray<FPolygonID> PolygonIDs;
+	PolygonIDs.Reset();
+
+	// Get a list of all polygon refs in the mesh
+	// @todo mesheditor: provide direct access to the polygons array in the group
+	{
+		const int32 PolygonCountInGroup = GetPolygonCountInGroup( PolygonGroupID );
+		PolygonIDs.Reserve( PolygonCountInGroup );
+		for( int32 PolygonIndex = 0; PolygonIndex < PolygonCountInGroup; ++PolygonIndex )
+		{
+			PolygonIDs.Add( GetPolygonInGroup( PolygonGroupID, PolygonIndex ) );
+		}
+	}
+
+	// This represents an adjacent triangle which can be merged to a quadrilateral, and an assigned score based on the 'quality' of the resulting quadrilateral.
+	struct FAdjacentPolygon
+	{
+		// Polygon ID of the adjacent triangle. This object is keyed on 'our' Polygon ID.
+		FPolygonID PolygonID;
+
+		// List of the vertices which form the merged quadrilateral - (Polygon ID, PolygonPerimeterIndex)
+		TTuple<FPolygonID, int32> Vertices[ 4 ];
+
+		// 'Quality' of the quadrilateral (internal angles closer to 90 degrees are better)
+		float Score;
+
+		FAdjacentPolygon() {}
+
+		FAdjacentPolygon( FPolygonID InPolygonID, FPolygonID InAdjacentPolygonID, int32 InVertex0, int32 InVertex1, int32 InVertex2, int32 InVertex3, float InScore )
+		{
+			PolygonID = InAdjacentPolygonID;
+			Vertices[ 0 ] = MakeTuple( InPolygonID, InVertex0 );
+			Vertices[ 1 ] = MakeTuple( InPolygonID, InVertex1 );
+			Vertices[ 2 ] = MakeTuple( InAdjacentPolygonID, InVertex2 );
+			Vertices[ 3 ] = MakeTuple( InAdjacentPolygonID, InVertex3 );
+			Score = InScore;
+		}
+	};
+
+
+	// This represents a list of adjacent polygons, ordered by score.
+	// Since we are only connecting triangles, there are a maximum of three adjacent polygons.
+	struct FAdjacentPolygons
+	{
+		enum { MaxAdjacentPolygons = 3 };
+
+		FAdjacentPolygon AdjacentPolygons[ MaxAdjacentPolygons ];
+		int32 NumAdjacentPolygons;
+
+		FAdjacentPolygons()
+		{
+			NumAdjacentPolygons = 0;
+		}
+
+		void Add( const FAdjacentPolygon& AdjacentPolygon )
+		{
+			check( NumAdjacentPolygons < MaxAdjacentPolygons );
+			int32 InsertIndex = 0;
+			for( int32 Index = 0; Index < NumAdjacentPolygons; ++Index )
+			{
+				if( AdjacentPolygon.Score > AdjacentPolygons[ Index ].Score )
+				{
+					InsertIndex++;
+				}
+				else
+				{
+					break;
+				}
+			}
+
+			for( int32 Index = NumAdjacentPolygons; Index > InsertIndex; --Index )
+			{
+				AdjacentPolygons[ Index ] = AdjacentPolygons[ Index - 1 ];
+			}
+			AdjacentPolygons[ InsertIndex ] = AdjacentPolygon;
+			NumAdjacentPolygons++;
+		}
+
+		const FAdjacentPolygon& GetBestAdjacentPolygon() const
+		{
+			check( NumAdjacentPolygons > 0 );
+			return AdjacentPolygons[ 0 ];
+		}
+		
+		bool Remove( const FPolygonID PolygonID )
+		{
+			for( int32 Index = 0; Index < NumAdjacentPolygons; ++Index )
+			{
+				if( AdjacentPolygons[ Index ].PolygonID == PolygonID )
+				{
+					for( int32 CopyIndex = Index + 1; CopyIndex < NumAdjacentPolygons; ++CopyIndex )
+					{
+						AdjacentPolygons[ CopyIndex - 1 ] = AdjacentPolygons[ CopyIndex ];
+					}
+
+					NumAdjacentPolygons--;
+					return true;
+				}
+			}
+
+			return false;
+		}
+
+		bool Contains( const FPolygonID PolygonID ) const
+		{
+			for( int32 Index = 0; Index < NumAdjacentPolygons; ++Index )
+			{
+				if( AdjacentPolygons[ Index ].PolygonID == PolygonID )
+				{
+					return true;
+				}
+			}
+
+			return false;
+		}
+
+		FPolygonID GetPolygonID( int32 Index ) const
+		{
+			check( Index < NumAdjacentPolygons );
+			return AdjacentPolygons[ Index ].PolygonID;
+		}
+
+		int32 Num() const { return NumAdjacentPolygons; }
+		bool IsValid() const { return NumAdjacentPolygons > 0; }
+	};
+
+
+	// Build list of valid adjacent triangle pairs, and assign a score based on the quality of the quadrilateral they form
+
+	static TMap<FPolygonID, FAdjacentPolygons> AdjacentPolygonsMap;
+	AdjacentPolygonsMap.Reset();
+
+	FPolygonID PolygonIDToMerge1 = FPolygonID::Invalid;
+	{
+		float BestScore = TNumericLimits<float>::Max();
+
+		for( const FPolygonID PolygonID : PolygonIDs )
+		{
+			// If it's not a triangle, don't consider this polygon at all
+			if( GetPolygonPerimeterEdgeCount( PolygonID ) != 3 )
+			{
+				continue;
+			}
+
+			// We're only interested in adjacent triangles which are nearly coplanar; get the normal so we can compare it with the adjacent polygons' normals
+			const FVector PolygonNormal = ComputePolygonNormal( PolygonID );
+
+			// Go round the edge considering all adjacent polygons, looking for valid pairs and assigning a quality score (lower is better)
+			for( int32 PerimeterEdgeIndex = 0; PerimeterEdgeIndex < 3; ++PerimeterEdgeIndex )
+			{
+				bool bOutEdgeWindingIsReversedForPolygon;
+				const FEdgeID PerimeterEdgeID = GetPolygonPerimeterEdge( PolygonID, PerimeterEdgeIndex, bOutEdgeWindingIsReversedForPolygon );
+
+				const bool bIsSoftEdge = FMath::IsNearlyZero( GetEdgeAttribute( PerimeterEdgeID, UEditableMeshAttribute::EdgeIsHard(), 0 ).X );
+				if( !bKeepHardEdges || bIsSoftEdge )
+				{
+					const FPolygonID AdjacentPolygonID = [ this, PerimeterEdgeID, PolygonID ]() -> FPolygonID
+					{
+						const int32 EdgeConnectedPolygonCount = GetEdgeConnectedPolygonCount( PerimeterEdgeID );
+
+						// Only interested in edges with exactly two connected polygons
+						if( EdgeConnectedPolygonCount == 2 )
+						{
+							for( int32 EdgeConnectedPolygonIndex = 0; EdgeConnectedPolygonIndex < 2; ++EdgeConnectedPolygonIndex )
+							{
+								const FPolygonID EdgeConnectedPolygonID = GetEdgeConnectedPolygon( PerimeterEdgeID, EdgeConnectedPolygonIndex );
+								if( EdgeConnectedPolygonID != PolygonID )
+								{
+									return GetPolygonPerimeterEdgeCount( EdgeConnectedPolygonID ) == 3 ? EdgeConnectedPolygonID : FPolygonID::Invalid;
+								}
+							}
+						}
+
+						return FPolygonID::Invalid;
+					}();
+
+					if( AdjacentPolygonID != FPolygonID::Invalid )
+					{
+						FAdjacentPolygons* AdjacentPolygons = AdjacentPolygonsMap.Find( PolygonID );
+
+						if( !AdjacentPolygons || !AdjacentPolygons->Contains( AdjacentPolygonID ) )
+						{
+							const FVector AdjacentPolygonNormal = ComputePolygonNormal( AdjacentPolygonID );
+							const float AdjacentPolygonDot = FVector::DotProduct( PolygonNormal, AdjacentPolygonNormal );
+
+							if( AdjacentPolygonDot >= CosAngleThreshold )
+							{
+								// Found a valid triangle pair whose interplanar angle is sufficiently shallow;
+								// now calculate a score according to the internal angles of the resulting quad
+
+								// We consider points on the two triangles' perimeters which form the quadrilateral.
+								// If the shared edge of the adjacent triangles falls on perimeter vertex N1 of triangle 1,
+								// and perimeter vertex N2 of triangle 2, then the points we consider are:
+								//
+								// (triangle 1, point N1 - 1)
+								// (triangle 1, point N1)
+								// (triangle 2, point N2 + 1)
+								// (triangle 2, point N2 - 1)
+								//
+								// or, from the perspective of the other triangle:
+								//
+								// (triangle 2, point N2 + 1)
+								// (triangle 2, point N2 - 1)
+								// (triangle 1, point N1 - 1)
+								// (triangle 1, point N1)
+
+								const int32 PrevPerimeterEdgeIndex = ( PerimeterEdgeIndex + 2 ) % 3;
+								const int32 NextPerimeterEdgeIndex = ( PerimeterEdgeIndex + 1 ) % 3;
+
+								const FVertexID SharedVertexID = GetPolygonPerimeterVertex( PolygonID, PerimeterEdgeIndex );
+								const int32 AdjacentPerimeterEdgeIndex = FindPolygonPerimeterVertexNumberForVertex( AdjacentPolygonID, SharedVertexID );
+
+								const int32 PrevAdjacentPerimeterEdgeIndex = ( AdjacentPerimeterEdgeIndex + 2 ) % 3;
+								const int32 NextAdjacentPerimeterEdgeIndex = ( AdjacentPerimeterEdgeIndex + 1 ) % 3;
+
+								// If we wish to maintain texture seams, compare the vertex UVs of each polygon's vertex along the edge to be removed, and only proceed if they are equal
+								bool bTextureCoordinatesEqual = true;
+								if( bKeepTextureBorder )
+								{
+									for( int32 TextureCoordinateIndex = 0; TextureCoordinateIndex < TextureCoordinateCount; ++TextureCoordinateIndex )
+									{
+										const FVector4 UVPolygon1Vertex1 = GetPolygonPerimeterVertexAttribute( PolygonID, PerimeterEdgeIndex, UEditableMeshAttribute::VertexTextureCoordinate(), TextureCoordinateIndex );
+										const FVector4 UVPolygon2Vertex1 = GetPolygonPerimeterVertexAttribute( AdjacentPolygonID, AdjacentPerimeterEdgeIndex, UEditableMeshAttribute::VertexTextureCoordinate(), TextureCoordinateIndex );
+										const FVector4 UVPolygon1Vertex2 = GetPolygonPerimeterVertexAttribute( PolygonID, NextPerimeterEdgeIndex, UEditableMeshAttribute::VertexTextureCoordinate(), TextureCoordinateIndex );
+										const FVector4 UVPolygon2Vertex2 = GetPolygonPerimeterVertexAttribute( AdjacentPolygonID, PrevAdjacentPerimeterEdgeIndex, UEditableMeshAttribute::VertexTextureCoordinate(), TextureCoordinateIndex );
+										if( UVPolygon1Vertex1 != UVPolygon2Vertex1 || UVPolygon1Vertex2 != UVPolygon2Vertex2 )
+										{
+											bTextureCoordinatesEqual = false;
+											break;
+										}
+									}
+								}
+
+								// If we wish to maintain color seams, compare the vertex colors of each polygon's vertex along the edge to be removed, and only proceed if they are equal
+								bool bColorsEqual = true;
+								if( bKeepColorBorder )
+								{
+									const FVector4 ColorPolygon1Vertex1 = GetPolygonPerimeterVertexAttribute( PolygonID, PerimeterEdgeIndex, UEditableMeshAttribute::VertexColor(), 0 );
+									const FVector4 ColorPolygon2Vertex1 = GetPolygonPerimeterVertexAttribute( AdjacentPolygonID, AdjacentPerimeterEdgeIndex, UEditableMeshAttribute::VertexColor(), 0 );
+									const FVector4 ColorPolygon1Vertex2 = GetPolygonPerimeterVertexAttribute( PolygonID, NextPerimeterEdgeIndex, UEditableMeshAttribute::VertexColor(), 0 );
+									const FVector4 ColorPolygon2Vertex2 = GetPolygonPerimeterVertexAttribute( AdjacentPolygonID, PrevAdjacentPerimeterEdgeIndex, UEditableMeshAttribute::VertexColor(), 0 );
+									if( ColorPolygon1Vertex1 != ColorPolygon2Vertex1 || ColorPolygon1Vertex2 != ColorPolygon2Vertex2 )
+									{
+										bColorsEqual = false;
+									}
+								}
+
+								// If we are maintaining texture or color seams, only merge the polygons if the appropriate attributes are equal for each polygon
+								if( ( !bKeepTextureBorder || bTextureCoordinatesEqual ) && ( !bKeepColorBorder || bColorsEqual ) )
+								{
+									const FVertexID V0 = GetPolygonPerimeterVertex( PolygonID, PrevPerimeterEdgeIndex );
+									const FVertexID V1 = GetPolygonPerimeterVertex( PolygonID, PerimeterEdgeIndex );
+									const FVertexID V2 = GetPolygonPerimeterVertex( AdjacentPolygonID, NextAdjacentPerimeterEdgeIndex );
+									const FVertexID V3 = GetPolygonPerimeterVertex( AdjacentPolygonID, PrevAdjacentPerimeterEdgeIndex );
+									check( V3 == GetPolygonPerimeterVertex( PolygonID, NextPerimeterEdgeIndex ) );
+
+									const FVector P0 = GetVertexAttribute( V0, UEditableMeshAttribute::VertexPosition(), 0 );
+									const FVector P1 = GetVertexAttribute( V1, UEditableMeshAttribute::VertexPosition(), 0 );
+									const FVector P2 = GetVertexAttribute( V2, UEditableMeshAttribute::VertexPosition(), 0 );
+									const FVector P3 = GetVertexAttribute( V3, UEditableMeshAttribute::VertexPosition(), 0 );
+
+									const FVector D01 = ( P1 - P0 ).GetSafeNormal();
+									const FVector D12 = ( P2 - P1 ).GetSafeNormal();
+									const FVector D23 = ( P3 - P2 ).GetSafeNormal();
+									const FVector D30 = ( P0 - P3 ).GetSafeNormal();
+
+									// Calculate a score based on the internal angles of the quadrilateral and the interplanar angle.
+									// Internal angles close to 90 degrees, and an interplanar angle close to 180 degrees are ideal.
+									const float Score =
+										FMath::Abs( HALF_PI - FMath::Acos( FVector::DotProduct( -D30, D01 ) ) ) +
+										FMath::Abs( HALF_PI - FMath::Acos( FVector::DotProduct( -D01, D12 ) ) ) +
+										FMath::Abs( HALF_PI - FMath::Acos( FVector::DotProduct( -D12, D23 ) ) ) +
+										FMath::Abs( HALF_PI - FMath::Acos( FVector::DotProduct( -D23, D30 ) ) ) +
+										FMath::Acos( AdjacentPolygonDot );
+
+									if( Score < BestScore )
+									{
+										BestScore = Score;
+										PolygonIDToMerge1 = PolygonID;
+									}
+
+									// Add to a list of adjacent polygons, sorted by score
+									FAdjacentPolygons& AdjacentPolygons1 = AdjacentPolygonsMap.FindOrAdd( PolygonID );
+									AdjacentPolygons1.Add( FAdjacentPolygon(
+										PolygonID,
+										AdjacentPolygonID,
+										PrevPerimeterEdgeIndex,
+										PerimeterEdgeIndex,
+										NextAdjacentPerimeterEdgeIndex,
+										PrevAdjacentPerimeterEdgeIndex,
+										Score ) );
+
+									// And perform the corresponding operation the other way round
+									FAdjacentPolygons& AdjacentPolygons2 = AdjacentPolygonsMap.FindOrAdd( AdjacentPolygonID );
+									check( !AdjacentPolygons2.Contains( PolygonID ) );
+									AdjacentPolygons2.Add( FAdjacentPolygon(
+										AdjacentPolygonID,
+										PolygonID,
+										NextAdjacentPerimeterEdgeIndex,
+										PrevAdjacentPerimeterEdgeIndex,
+										PrevPerimeterEdgeIndex,
+										PerimeterEdgeIndex,
+										Score ) );
+								}
+							}
+						}
+					}
+				}
+			}
+		}
+	}
+
+	// If there were no valid pairs of polys to merge, finish already
+	if( PolygonIDToMerge1 == FPolygonID::Invalid )
+	{
+		return;
+	}
+
+	static TArray<FPolygonToCreate> PolygonsToCreate;
+	PolygonsToCreate.Reset();
+
+	static TArray<FPolygonID> PolygonIDsToDelete;
+	PolygonIDsToDelete.Reset();
+
+	static TSet<FPolygonID> BoundaryPolygons;
+	BoundaryPolygons.Reset();
+
+	// Propagate quadrangulated area outwards from starting polygon
+	for( ; ; )
+	{
+		FAdjacentPolygons& AdjacentPolygons1 = AdjacentPolygonsMap.FindChecked( PolygonIDToMerge1 );
+		check( AdjacentPolygons1.IsValid() );
+		const FAdjacentPolygon& AdjacentPolygon1 = AdjacentPolygons1.GetBestAdjacentPolygon();
+		const FPolygonID PolygonIDToMerge2 = AdjacentPolygon1.PolygonID;
+
+		FAdjacentPolygons& AdjacentPolygons2 = AdjacentPolygonsMap.FindChecked( PolygonIDToMerge2 );
+		check( AdjacentPolygons2.IsValid() );
+
+		// Create new quadrilateral
+
+		PolygonsToCreate.Emplace();
+
+		FPolygonToCreate& PolygonToCreate = PolygonsToCreate.Last();
+
+		PolygonToCreate.PolygonGroupID = PolygonGroupID;
+		PolygonToCreate.PerimeterVertices.Reset( 4 );
+
+		for( int32 VertexIndex = 0; VertexIndex < 4; ++VertexIndex )
+		{
+			PolygonToCreate.PerimeterVertices.Emplace();
+			FVertexAndAttributes& VertexAndAttributes = PolygonToCreate.PerimeterVertices.Last();
+
+			VertexAndAttributes.VertexID = GetPolygonPerimeterVertex(
+				AdjacentPolygon1.Vertices[ VertexIndex ].Get<0>(),
+				AdjacentPolygon1.Vertices[ VertexIndex ].Get<1>() );
+
+			TArray<FMeshElementAttributeData>& Attributes = VertexAndAttributes.PolygonVertexAttributes.Attributes;
+
+			for( int32 TextureCoordinateIndex = 0; TextureCoordinateIndex < TextureCoordinateCount; ++TextureCoordinateIndex )
+			{
+				Attributes.Emplace(
+					UEditableMeshAttribute::VertexTextureCoordinate(),
+					TextureCoordinateIndex,
+					GetPolygonPerimeterVertexAttribute( 
+						AdjacentPolygon1.Vertices[ VertexIndex ].Get<0>(),
+						AdjacentPolygon1.Vertices[ VertexIndex ].Get<1>(),
+						UEditableMeshAttribute::VertexTextureCoordinate(),
+						TextureCoordinateIndex ) 
+					);
+			}
+
+			Attributes.Emplace(
+				UEditableMeshAttribute::VertexColor(),
+				0,
+				GetPolygonPerimeterVertexAttribute(
+					AdjacentPolygon1.Vertices[ VertexIndex ].Get<0>(),
+					AdjacentPolygon1.Vertices[ VertexIndex ].Get<1>(),
+					UEditableMeshAttribute::VertexColor(),
+					0 )
+				);
+		}
+
+		// Specify old polygons to be deleted
+
+		check( !PolygonIDsToDelete.Contains( PolygonIDToMerge1 ) );
+		check( !PolygonIDsToDelete.Contains( PolygonIDToMerge2 ) );
+		PolygonIDsToDelete.Emplace( PolygonIDToMerge1 );
+		PolygonIDsToDelete.Emplace( PolygonIDToMerge2 );
+
+		// And remove them from the boundary set
+
+		BoundaryPolygons.Remove( PolygonIDToMerge1 );
+		BoundaryPolygons.Remove( PolygonIDToMerge2 );
+
+		// Now break connections between newly added polygons and their neighbors.
+		// If a polygon ends up with no connections, delete it entirely from the map so it is no longer considered.
+		// This happens if a polygon has been added to the quadrangulated set, or if it is an orphaned triangle which cannot be paired to anything.
+		// We defer deleting the entry from the map until we have broken all connections.
+
+		verify( AdjacentPolygons1.Remove( PolygonIDToMerge2 ) );
+		verify( AdjacentPolygons2.Remove( PolygonIDToMerge1 ) );
+
+		static TArray<FPolygonID> AdjacentPolygonsEntryToDelete;
+		AdjacentPolygonsEntryToDelete.Reset();
+
+		for( int32 Index = 0; Index < AdjacentPolygons1.Num(); ++Index )
+		{
+			const FPolygonID AdjacentPolygonID = AdjacentPolygons1.GetPolygonID( Index );
+
+			FAdjacentPolygons* OtherAdjacentPolygons = AdjacentPolygonsMap.Find( AdjacentPolygonID );
+			if( OtherAdjacentPolygons )
+			{
+				verify( OtherAdjacentPolygons->Remove( PolygonIDToMerge1 ) );
+				if( !OtherAdjacentPolygons->IsValid() )
+				{
+					AdjacentPolygonsEntryToDelete.Add( AdjacentPolygonID );
+				}
+				else
+				{
+					BoundaryPolygons.Add( AdjacentPolygonID );
+				}
+			}
+		}
+
+		AdjacentPolygonsEntryToDelete.Add( PolygonIDToMerge1 );
+
+		for( int32 Index = 0; Index < AdjacentPolygons2.Num(); ++Index )
+		{
+			const FPolygonID AdjacentPolygonID = AdjacentPolygons2.GetPolygonID( Index );
+
+			FAdjacentPolygons* OtherAdjacentPolygons = AdjacentPolygonsMap.Find( AdjacentPolygonID );
+			if( OtherAdjacentPolygons )
+			{
+				verify( OtherAdjacentPolygons->Remove( PolygonIDToMerge2 ) );
+				if( !OtherAdjacentPolygons->IsValid() )
+				{
+					AdjacentPolygonsEntryToDelete.Add( AdjacentPolygonID );
+				}
+				else
+				{
+					BoundaryPolygons.Add( AdjacentPolygonID );
+				}
+			}
+		}
+
+		AdjacentPolygonsEntryToDelete.Add( PolygonIDToMerge2 );
+
+		// Clean up: any polygons' map entries which now have no adjacent polygons get deleted completely.
+		// This implies they have no connected neighbors which can be merged (either because they are near an edge with only unmergeable polygons nearby, or because
+		// they are in the middle of the quadrangulated area)
+
+		for( const FPolygonID AdjacentPolygonsEntry : AdjacentPolygonsEntryToDelete )
+		{
+			AdjacentPolygonsMap.Remove( AdjacentPolygonsEntry );
+			BoundaryPolygons.Remove( AdjacentPolygonsEntry );
+		}
+
+		// Now look for the next polygon to use: it is the one with the best score from the BoundaryPolygons set.
+
+		float BestScore = TNumericLimits<float>::Max();
+		PolygonIDToMerge1 = FPolygonID::Invalid;
+		for( const FPolygonID BoundaryPolygon : BoundaryPolygons )
+		{
+			const FAdjacentPolygons& AdjacentPolygons = AdjacentPolygonsMap.FindChecked( BoundaryPolygon );
+			const FAdjacentPolygon& AdjacentPolygon = AdjacentPolygons.GetBestAdjacentPolygon();
+			if( AdjacentPolygon.Score < BestScore )
+			{
+				PolygonIDToMerge1 = AdjacentPolygon.PolygonID;
+			}
+		}
+
+		// If there are still no candidates adjacent to the already quadrangulated area, choose the best candidate elsewhere.
+		// This will start a new quadrangulated area, which is grown in the same way as the last.
+
+		if( PolygonIDToMerge1 == FPolygonID::Invalid )
+		{
+			BoundaryPolygons.Reset();
+
+			for( const auto& AdjacentPolygonsMapEntry : AdjacentPolygonsMap )
+			{
+				const FPolygonID PolygonID = AdjacentPolygonsMapEntry.Key;
+				const FAdjacentPolygons& AdjacentPolygons = AdjacentPolygonsMapEntry.Value;
+				check( AdjacentPolygons.IsValid() );
+
+				const FAdjacentPolygon& AdjacentPolygon = AdjacentPolygons.GetBestAdjacentPolygon();
+
+				if( AdjacentPolygon.Score < BestScore )
+				{
+					BestScore = AdjacentPolygon.Score;
+					PolygonIDToMerge1 = PolygonID;
+				}
+			}
+		}
+
+		// If there are still no candidates, we've done as much as we can do
+
+		if( PolygonIDToMerge1 == FPolygonID::Invalid )
+		{
+			break;
+		}
+	}
+
+	// Finally, actually change the geometry and rebuild normals/tangents
+
+	static TArray<FEdgeID> CreatedEdgeIDs;
+
+	CreatePolygons( PolygonsToCreate, OutNewPolygonIDs, CreatedEdgeIDs );
+
+	const bool bDeleteOrphanedEdges = true;
+	const bool bDeleteOrphanedVertices = false;
+	const bool bDeleteOrphanedVertexInstances = false;
+	const bool bDeleteEmptyPolygonGroups = false;
+	DeletePolygons( PolygonIDsToDelete, bDeleteOrphanedEdges, bDeleteOrphanedVertices, bDeleteOrphanedVertexInstances, bDeleteEmptyPolygonGroups );
+}
+
+
+bool UEditableMesh::AnyChangesToUndo() const
+{
+	return bAllowUndo && Undo.IsValid() && Undo->Subchanges.Num() > 0;
+}
+
+
+void UEditableMesh::AddUndo( TUniquePtr<FChange> NewUndo )
+{
+	if( NewUndo.IsValid() )
+	{
+//		UE_LOG(LogEditableMesh, Verbose, TEXT("%s"), *NewUndo->ToString());
+		if( bAllowUndo )
+		{
+			if( !Undo.IsValid() )
+			{
+				Undo = MakeUnique<FCompoundChangeInput>();
+			}
+
+			Undo->Subchanges.Add( MoveTemp( NewUndo ) );
+		}
+	}
+}
+
+
+TUniquePtr<FChange> UEditableMesh::MakeUndo()
+{
+	TUniquePtr<FChange> UndoChange = nullptr;
+	if( AnyChangesToUndo() )
+	{
+		UndoChange = MakeUnique<FCompoundChange>( MoveTemp( *Undo ) );
+	}
+	Undo.Reset();
+
+	return UndoChange;
+}
\ No newline at end of file
