Mesh Editor: New "EditableMesh" plugin (replaces "MeshEditingRuntime" module)

diff --git a/Engine/Plugins/Runtime/EditableMesh/Source/EditableMesh/EditableMeshChanges.cpp b/Engine/Plugins/Runtime/EditableMesh/Source/EditableMesh/EditableMeshChanges.cpp
new file mode 100644
--- /dev/null
+++ b/Engine/Plugins/Runtime/EditableMesh/Source/EditableMesh/EditableMeshChanges.cpp
@@ -0,0 +1,399 @@
+// Copyright 1998-2015 Epic Games, Inc. All Rights Reserved.
+
+#include "EditableMeshChanges.h"
+#include "EditableMesh.h"
+
+
+TUniquePtr<FChange> FDeleteOrphanVerticesChange::Execute( UObject* Object )
+{
+	UEditableMesh* EditableMesh = CastChecked<UEditableMesh>( Object );
+	verify( !EditableMesh->AnyChangesToUndo() );
+
+	EditableMesh->DeleteOrphanVertices( Input.VertexIDsToDelete );
+
+	return EditableMesh->MakeUndo();
+}
+
+
+FString FDeleteOrphanVerticesChange::ToString() const
+{
+	return FString::Printf(
+		TEXT( "Delete Orphan Vertices [VertexIDsToDelete:%s]" ),
+		*LogHelpers::ArrayToString( Input.VertexIDsToDelete ) );
+}
+
+
+TUniquePtr<FChange> FDeleteVertexInstancesChange::Execute( UObject* Object )
+{
+	UEditableMesh* EditableMesh = CastChecked<UEditableMesh>( Object );
+	verify( !EditableMesh->AnyChangesToUndo() );
+
+	EditableMesh->DeleteVertexInstances( Input.VertexInstanceIDsToDelete, Input.bDeleteOrphanedVertices );
+
+	return EditableMesh->MakeUndo();
+}
+
+
+FString FDeleteVertexInstancesChange::ToString() const
+{
+	return FString::Printf(
+		TEXT( "Delete Vertex Instances [VertexInstanceIDsToDelete:%s]" ),
+		*LogHelpers::ArrayToString( Input.VertexInstanceIDsToDelete ) );
+}
+
+
+TUniquePtr<FChange> FDeleteEdgesChange::Execute( UObject* Object )
+{
+	UEditableMesh* EditableMesh = CastChecked<UEditableMesh>( Object );
+	verify( !EditableMesh->AnyChangesToUndo() );
+
+	EditableMesh->DeleteEdges( Input.EdgeIDsToDelete, Input.bDeleteOrphanedVertices );
+
+	return EditableMesh->MakeUndo();
+}
+
+
+FString FDeleteEdgesChange::ToString() const
+{
+	return FString::Printf(
+		TEXT( "Delete Edges [EdgeIDsToDelete:%s, bDeleteOrphanedVertices:%s]" ),
+		*LogHelpers::ArrayToString( Input.EdgeIDsToDelete ),
+		*LogHelpers::BoolToString( Input.bDeleteOrphanedVertices ) );
+}
+
+
+TUniquePtr<FChange> FCreateVerticesChange::Execute( UObject* Object )
+{
+	UEditableMesh* EditableMesh = CastChecked<UEditableMesh>( Object );
+	verify( !EditableMesh->AnyChangesToUndo() );
+
+	static TArray<FVertexID> UnusedNewVertexIDs;
+	EditableMesh->CreateVertices( Input.VerticesToCreate, /* Out */ UnusedNewVertexIDs );
+
+	return EditableMesh->MakeUndo();
+}
+
+
+FString FCreateVerticesChange::ToString() const
+{
+	return FString::Printf(
+		TEXT( "Create Vertices [VerticesToCreate:%s]" ),
+		*LogHelpers::ArrayToString( Input.VerticesToCreate ) );
+}
+
+
+TUniquePtr<FChange> FCreateVertexInstancesChange::Execute( UObject* Object )
+{
+	UEditableMesh* EditableMesh = CastChecked<UEditableMesh>( Object );
+	verify( !EditableMesh->AnyChangesToUndo() );
+
+	static TArray<FVertexInstanceID> UnusedNewVertexInstanceIDs;
+	EditableMesh->CreateVertexInstances( Input.VertexInstancesToCreate, /* Out */ UnusedNewVertexInstanceIDs );
+
+	return EditableMesh->MakeUndo();
+}
+
+
+FString FCreateVertexInstancesChange::ToString() const
+{
+	return FString::Printf(
+		TEXT( "Create Vertex Instances [VertexInstancesToCreate:%s]" ),
+		*LogHelpers::ArrayToString( Input.VertexInstancesToCreate ) );
+}
+
+
+TUniquePtr<FChange> FCreateEdgesChange::Execute( UObject* Object )
+{
+	UEditableMesh* EditableMesh = CastChecked<UEditableMesh>( Object );
+	verify( !EditableMesh->AnyChangesToUndo() );
+
+	static TArray<FEdgeID> UnusedNewEdgeIDs;
+	EditableMesh->CreateEdges( Input.EdgesToCreate, /* Out */ UnusedNewEdgeIDs );
+
+	return EditableMesh->MakeUndo();
+}
+
+
+FString FCreateEdgesChange::ToString() const
+{
+	return FString::Printf(
+		TEXT( "Create Edges [EdgesToCreate:%s]" ),
+		*LogHelpers::ArrayToString( Input.EdgesToCreate ) );
+}
+
+
+TUniquePtr<FChange> FCreatePolygonsChange::Execute( UObject* Object )
+{
+	UEditableMesh* EditableMesh = CastChecked<UEditableMesh>( Object );
+	verify( !EditableMesh->AnyChangesToUndo() );
+	
+	static TArray<FPolygonID> UnusedNewPolygonIDs;
+	static TArray<FEdgeID> UnusedNewEdgeIDs;
+	EditableMesh->CreatePolygons( Input.PolygonsToCreate, /* Out */ UnusedNewPolygonIDs, /* Out */ UnusedNewEdgeIDs );
+
+	return EditableMesh->MakeUndo();
+}
+
+
+FString FCreatePolygonsChange::ToString() const
+{
+	return FString::Printf(
+		TEXT( "Create Polygons [PolygonsToCreate:%s]" ),
+		*LogHelpers::ArrayToString( Input.PolygonsToCreate ) );
+}
+
+
+TUniquePtr<FChange> FDeletePolygonsChange::Execute( UObject* Object )
+{
+	UEditableMesh* EditableMesh = CastChecked<UEditableMesh>( Object );
+	verify( !EditableMesh->AnyChangesToUndo() );
+
+	EditableMesh->DeletePolygons( Input.PolygonIDsToDelete, Input.bDeleteOrphanedEdges, Input.bDeleteOrphanedVertices, Input.bDeleteOrphanedVertexInstances, Input.bDeleteEmptySections );
+	return EditableMesh->MakeUndo();
+}
+
+
+FString FDeletePolygonsChange::ToString() const
+{
+	return FString::Printf(
+		TEXT( "Delete Polygons [PolygonIDsToDelete:%s, %s, %s]" ),
+		*LogHelpers::ArrayToString( Input.PolygonIDsToDelete ),
+		*LogHelpers::BoolToString( Input.bDeleteOrphanedEdges ),
+		*LogHelpers::BoolToString( Input.bDeleteOrphanedVertices ),
+		*LogHelpers::BoolToString( Input.bDeleteEmptySections ) );
+}
+
+
+TUniquePtr<FChange> FSetVerticesAttributesChange::Execute( UObject* Object )
+{
+	UEditableMesh* EditableMesh = CastChecked<UEditableMesh>( Object );
+	verify( !EditableMesh->AnyChangesToUndo() );
+
+	EditableMesh->SetVerticesAttributes( Input.AttributesForVertices );
+	return EditableMesh->MakeUndo();
+}
+
+
+FString FSetVerticesAttributesChange::ToString() const
+{
+	return FString::Printf(
+		TEXT( "Set Vertices Attributes [AttributesForVertices:%s]" ),
+		*LogHelpers::ArrayToString( Input.AttributesForVertices ) );
+}
+
+
+TUniquePtr<FChange> FSetVertexInstancesAttributesChange::Execute( UObject* Object )
+{
+	UEditableMesh* EditableMesh = CastChecked<UEditableMesh>( Object );
+	verify( !EditableMesh->AnyChangesToUndo() );
+
+	EditableMesh->SetVertexInstancesAttributes( Input.AttributesForVertexInstances );
+	return EditableMesh->MakeUndo();
+}
+
+
+FString FSetVertexInstancesAttributesChange::ToString() const
+{
+	return FString::Printf(
+		TEXT( "Set Vertex Instances Attributes [AttributesForVertexInstances:%s]" ),
+		*LogHelpers::ArrayToString( Input.AttributesForVertexInstances ) );
+}
+
+
+TUniquePtr<FChange> FSetEdgesAttributesChange::Execute( UObject* Object )
+{
+	UEditableMesh* EditableMesh = CastChecked<UEditableMesh>( Object );
+	verify( !EditableMesh->AnyChangesToUndo() );
+
+	EditableMesh->SetEdgesAttributes( Input.AttributesForEdges );
+	return EditableMesh->MakeUndo();
+}
+
+
+FString FSetEdgesAttributesChange::ToString() const
+{
+	return FString::Printf(
+		TEXT( "Set Edges Attributes [AttributesForEdges:%s]" ),
+		*LogHelpers::ArrayToString( Input.AttributesForEdges ) );
+}
+
+
+TUniquePtr<FChange> FSetPolygonsVertexAttributesChange::Execute( UObject* Object )
+{
+	UEditableMesh* EditableMesh = CastChecked<UEditableMesh>( Object );
+	verify( !EditableMesh->AnyChangesToUndo() );
+
+	EditableMesh->SetPolygonsVertexAttributes( Input.VertexAttributesForPolygons );
+	return EditableMesh->MakeUndo();
+}
+
+
+FString FSetPolygonsVertexAttributesChange::ToString() const
+{
+	return FString::Printf(
+		TEXT( "Set Polygons Vertex Attributes [VertexAttributesForPolygons:%s]" ),
+		*LogHelpers::ArrayToString( Input.VertexAttributesForPolygons ) );
+}
+
+
+TUniquePtr<FChange> FChangePolygonsVertexInstancesChange::Execute( UObject* Object )
+{
+	UEditableMesh* EditableMesh = CastChecked<UEditableMesh>( Object );
+	verify( !EditableMesh->AnyChangesToUndo() );
+
+	EditableMesh->ChangePolygonsVertexInstances( Input.VertexInstancesForPolygons );
+	return EditableMesh->MakeUndo();
+}
+
+
+FString FChangePolygonsVertexInstancesChange::ToString() const
+{
+	return FString::Printf(
+		TEXT( "Change Polygons Vertex Instances [VertexInstancesForPolygons:%s]" ),
+		*LogHelpers::ArrayToString( Input.VertexInstancesForPolygons ) );
+}
+
+
+TUniquePtr<FChange> FSetEdgesVerticesChange::Execute( UObject* Object )
+{
+	UEditableMesh* EditableMesh = CastChecked<UEditableMesh>( Object );
+	verify( !EditableMesh->AnyChangesToUndo() );
+
+	EditableMesh->SetEdgesVertices( Input.VerticesForEdges );
+	return EditableMesh->MakeUndo();
+}
+
+
+
+FString FSetEdgesVerticesChange::ToString() const
+{
+	return FString::Printf(
+		TEXT( "Set Edges Vertices [VerticesForEdges:%s]" ),
+		*LogHelpers::ArrayToString( Input.VerticesForEdges ) );
+}
+
+
+TUniquePtr<FChange> FInsertPolygonPerimeterVerticesChange::Execute( UObject* Object )
+{
+	UEditableMesh* EditableMesh = CastChecked<UEditableMesh>( Object );
+	verify( !EditableMesh->AnyChangesToUndo() );
+
+	EditableMesh->InsertPolygonPerimeterVertices( Input.PolygonID, Input.InsertBeforeVertexNumber, Input.VerticesToInsert );
+	return EditableMesh->MakeUndo();
+}
+
+
+FString FInsertPolygonPerimeterVerticesChange::ToString() const
+{
+	return FString::Printf(
+		TEXT( "Insert Polygon Perimeter Vertices [PolygonID:%s, InsertBeforeVertexNumber:%lu, VerticesToInsert:%s]" ),
+		*Input.PolygonID.ToString(),
+		Input.InsertBeforeVertexNumber,
+		*LogHelpers::ArrayToString( Input.VerticesToInsert ) );
+}
+
+
+TUniquePtr<FChange> FRemovePolygonPerimeterVerticesChange::Execute( UObject* Object )
+{
+	UEditableMesh* EditableMesh = CastChecked<UEditableMesh>( Object );
+	verify( !EditableMesh->AnyChangesToUndo() );
+
+	EditableMesh->RemovePolygonPerimeterVertices( Input.PolygonID, Input.FirstVertexNumberToRemove, Input.NumVerticesToRemove, Input.bDeleteOrphanedVertexInstances );
+	return EditableMesh->MakeUndo();
+}
+
+
+FString FRemovePolygonPerimeterVerticesChange::ToString() const
+{
+	return FString::Printf(
+		TEXT( "Remove Polygon Perimeter Vertices [PolygonID:%s, FirstVertexNumberToRemove:%lu, NumVerticesToRemove:%lu]" ),
+		*Input.PolygonID.ToString(),
+		Input.FirstVertexNumberToRemove,
+		Input.NumVerticesToRemove );
+}
+
+
+TUniquePtr<FChange> FStartOrEndModificationChange::Execute( UObject* Object )
+{
+	UEditableMesh* EditableMesh = CastChecked<UEditableMesh>( Object );
+	verify( !EditableMesh->AnyChangesToUndo() );
+
+	if( Input.bStartModification )
+	{
+		EditableMesh->StartModification( Input.MeshModificationType, Input.MeshTopologyChange );
+	}
+	else
+	{
+		const bool bFromUndo = true;
+		EditableMesh->EndModification( bFromUndo );
+	}
+
+	return EditableMesh->MakeUndo();
+}
+
+
+FString FStartOrEndModificationChange::ToString() const
+{
+	return FString::Printf(
+		TEXT( "%s Modification (MeshModificationType:%i, MeshTopologyChange:%i)" ),
+		Input.bStartModification ? TEXT( "Start" ) : TEXT( "End" ),
+		(int32)Input.MeshModificationType,
+		(int32)Input.MeshTopologyChange );
+}
+
+
+TUniquePtr<FChange> FSetSubdivisionCountChange::Execute( UObject* Object )
+{
+	UEditableMesh* EditableMesh = CastChecked<UEditableMesh>( Object );
+	verify( !EditableMesh->AnyChangesToUndo() );
+
+	EditableMesh->SetSubdivisionCount( Input.NewSubdivisionCount );
+	return EditableMesh->MakeUndo();
+}
+
+
+FString FSetSubdivisionCountChange::ToString() const
+{
+	return FString::Printf(
+		TEXT( "Set Subdivision Count [NewSubdivisionCount:%i]" ),
+		Input.NewSubdivisionCount );
+}
+
+
+TUniquePtr<FChange> FCreatePolygonGroupsChange::Execute( UObject* Object )
+{
+	UEditableMesh* EditableMesh = CastChecked<UEditableMesh>( Object );
+	verify( !EditableMesh->AnyChangesToUndo() );
+
+	static TArray<FPolygonGroupID> UnusedPolygonGroupIDs;
+	EditableMesh->CreatePolygonGroups( Input.PolygonGroupsToCreate, UnusedPolygonGroupIDs );
+	return EditableMesh->MakeUndo();
+}
+
+
+FString FCreatePolygonGroupsChange::ToString() const
+{
+	return FString::Printf(
+		TEXT( "Create Section [PolygonGroupsToCreate:%s]" ),
+		*LogHelpers::ArrayToString( Input.PolygonGroupsToCreate ) );
+}
+
+
+TUniquePtr<FChange> FDeletePolygonGroupsChange::Execute( UObject* Object )
+{
+	UEditableMesh* EditableMesh = CastChecked<UEditableMesh>( Object );
+	verify( !EditableMesh->AnyChangesToUndo() );
+
+	EditableMesh->DeletePolygonGroups( Input.PolygonGroupIDs );
+	return EditableMesh->MakeUndo();
+}
+
+
+FString FDeletePolygonGroupsChange::ToString() const
+{
+	return FString::Printf(
+		TEXT( "Delete Section [PolygonGroupIDs:%s]" ),
+		*LogHelpers::ArrayToString( Input.PolygonGroupIDs ) );
+}
+
+
diff --git a/Engine/Plugins/Runtime/EditableMesh/Source/EditableMesh/EditableMeshChanges.h b/Engine/Plugins/Runtime/EditableMesh/Source/EditableMesh/EditableMeshChanges.h
new file mode 100644
--- /dev/null
+++ b/Engine/Plugins/Runtime/EditableMesh/Source/EditableMesh/EditableMeshChanges.h
@@ -0,0 +1,860 @@
+// Copyright 1998-2015 Epic Games, Inc. All Rights Reserved.
+
+#pragma once
+
+#include "CoreMinimal.h"
+#include "Change.h"
+#include "EditableMeshTypes.h"
+
+
+// @todo mesheditor: Comment these classes and enums!
+
+struct FDeleteOrphanVerticesChangeInput
+{
+	/** The vertex IDs to get rid of */
+	TArray<FVertexID> VertexIDsToDelete;
+
+	/** Default constructor */
+	FDeleteOrphanVerticesChangeInput()
+		: VertexIDsToDelete()
+	{
+	}
+};
+
+
+class FDeleteOrphanVerticesChange : public FChange
+{
+
+public:
+
+	/** Constructor */
+	FDeleteOrphanVerticesChange( const FDeleteOrphanVerticesChangeInput& InitInput )
+		: Input( InitInput )
+	{
+	}
+
+ 	FDeleteOrphanVerticesChange( FDeleteOrphanVerticesChangeInput&& InitInput )
+ 		: Input( MoveTemp( InitInput ) )
+ 	{
+ 	}
+
+	// Parent class overrides
+	virtual TUniquePtr<FChange> Execute( UObject* Object ) override;
+	virtual FString ToString() const override;
+
+private:
+	
+	/** The data we need to make this change */
+	FDeleteOrphanVerticesChangeInput Input;
+
+};
+
+
+struct FDeleteVertexInstancesChangeInput
+{
+	/** The vertex instance IDs to delete */
+	TArray<FVertexInstanceID> VertexInstanceIDsToDelete;
+
+	/** Whether we should also delete any vertices if we delete their only instance */
+	bool bDeleteOrphanedVertices;
+
+	/** Default constructor */
+	FDeleteVertexInstancesChangeInput()
+		: VertexInstanceIDsToDelete(),
+		  bDeleteOrphanedVertices( true )
+	{
+	}
+};
+
+
+class FDeleteVertexInstancesChange : public FChange
+{
+
+public:
+
+	/** Constructor */
+	FDeleteVertexInstancesChange( const FDeleteVertexInstancesChangeInput& InitInput )
+		: Input( InitInput )
+	{
+	}
+
+	FDeleteVertexInstancesChange( FDeleteVertexInstancesChangeInput&& InitInput )
+		: Input( MoveTemp( InitInput ) )
+	{
+	}
+
+	// Parent class overrides
+	virtual TUniquePtr<FChange> Execute( UObject* Object ) override;
+	virtual FString ToString() const override;
+
+private:
+
+	/** The data we need to make this change */
+	FDeleteVertexInstancesChangeInput Input;
+};
+
+
+struct FDeleteEdgesChangeInput
+{
+	/** The edge IDs to get rid of */
+	TArray<FEdgeID> EdgeIDsToDelete;
+
+	/** Whether we should also delete any vertices that are left orphaned after deleting this edge */
+	bool bDeleteOrphanedVertices;
+
+	/** Default constructor */
+	FDeleteEdgesChangeInput()
+		: EdgeIDsToDelete(),
+		  bDeleteOrphanedVertices( true )
+	{
+	}
+};
+
+
+class FDeleteEdgesChange : public FChange
+{
+
+public:
+
+	/** Constructor */
+	FDeleteEdgesChange( const FDeleteEdgesChangeInput& InitInput )
+		: Input( InitInput )
+	{
+	}
+
+ 	FDeleteEdgesChange( FDeleteEdgesChangeInput&& InitInput )
+ 		: Input( MoveTemp( InitInput ) )
+ 	{
+ 	}
+
+	// Parent class overrides
+	virtual TUniquePtr<FChange> Execute( UObject* Object ) override;
+	virtual FString ToString() const override;
+
+
+private:
+
+	/** The data we need to make this change */
+	FDeleteEdgesChangeInput Input;
+
+};
+
+
+struct FCreateVerticesChangeInput
+{
+	/** Information about each vertex that will be created */
+	TArray<FVertexToCreate> VerticesToCreate;
+
+	/** Default constructor */
+	FCreateVerticesChangeInput()
+		: VerticesToCreate()
+	{
+	}
+};
+
+
+class FCreateVerticesChange : public FChange
+{
+
+public:
+
+	/** Constructor */
+	FCreateVerticesChange( const FCreateVerticesChangeInput& InitInput )
+		: Input( InitInput )
+	{
+	}
+
+ 	FCreateVerticesChange( FCreateVerticesChangeInput&& InitInput )
+ 		: Input( MoveTemp( InitInput ) )
+ 	{
+ 	}
+
+	// Parent class overrides
+	virtual TUniquePtr<FChange> Execute( UObject* Object ) override;
+	virtual FString ToString() const override;
+
+
+private:
+	
+	/** The data we need to make this change */
+	FCreateVerticesChangeInput Input;
+
+};
+
+
+struct FCreateVertexInstancesChangeInput
+{
+	/** Information about each vertex instance that will be created */
+	TArray<FVertexInstanceToCreate> VertexInstancesToCreate;
+
+	/** Default constructor */
+	FCreateVertexInstancesChangeInput()
+		: VertexInstancesToCreate()
+	{
+	}
+};
+
+
+class FCreateVertexInstancesChange : public FChange
+{
+
+public:
+
+	/** Constructor */
+	FCreateVertexInstancesChange( const FCreateVertexInstancesChangeInput& InitInput )
+		: Input( InitInput )
+	{
+	}
+
+	FCreateVertexInstancesChange( FCreateVertexInstancesChangeInput&& InitInput )
+		: Input( MoveTemp( InitInput ) )
+	{
+	}
+
+	// Parent class overrides
+	virtual TUniquePtr<FChange> Execute( UObject* Object ) override;
+	virtual FString ToString() const override;
+
+
+private:
+
+	/** The data we need to make this change */
+	FCreateVertexInstancesChangeInput Input;
+
+};
+
+
+struct FCreateEdgesChangeInput
+{
+	/** Information about each edge that will be created */
+	TArray<FEdgeToCreate> EdgesToCreate;
+
+	/** Default constructor */
+	FCreateEdgesChangeInput()
+		: EdgesToCreate()
+	{
+	}
+};
+
+
+class FCreateEdgesChange : public FChange
+{
+
+public:
+
+	/** Constructor */
+	FCreateEdgesChange( const FCreateEdgesChangeInput& InitInput )
+		: Input( InitInput )
+	{
+	}
+
+ 	FCreateEdgesChange( FCreateEdgesChangeInput&& InitInput )
+ 		: Input( MoveTemp( InitInput ) )
+ 	{
+ 	}
+
+	// Parent class overrides
+	virtual TUniquePtr<FChange> Execute( UObject* Object ) override;
+	virtual FString ToString() const override;
+
+
+private:
+
+	/** The data we need to make this change */
+	FCreateEdgesChangeInput Input;
+
+};
+
+
+struct FCreatePolygonsChangeInput
+{
+	/** Information about each polygon that will be created */
+	TArray<FPolygonToCreate> PolygonsToCreate;
+
+	/** Default constructor */
+	FCreatePolygonsChangeInput()
+		: PolygonsToCreate()
+	{
+	}
+};
+
+
+class FCreatePolygonsChange : public FChange
+{
+
+public:
+
+	/** Constructor */
+	FCreatePolygonsChange( const FCreatePolygonsChangeInput& InitInput )
+		: Input( InitInput )
+	{
+	}
+
+ 	FCreatePolygonsChange( FCreatePolygonsChangeInput&& InitInput )
+ 		: Input( MoveTemp( InitInput ) )
+ 	{
+ 	}
+
+	// Parent class overrides
+	virtual TUniquePtr<FChange> Execute( UObject* Object ) override;
+	virtual FString ToString() const override;
+
+
+private:
+
+	/** The data we need to make this change */
+	FCreatePolygonsChangeInput Input;
+
+};
+
+
+struct FDeletePolygonsChangeInput
+{
+	/** The polygons to get rid of */
+	TArray<FPolygonID> PolygonIDsToDelete;
+
+	/** Whether we should also delete any edges that are left orphaned after deleting this polygon */
+	bool bDeleteOrphanedEdges;
+
+	/** Whether we should also delete any vertices that are left orphaned after deleting this polygon */
+	bool bDeleteOrphanedVertices;
+
+	/** Whether we should also delete any vertex instances that are left orphaned after deleting this polygon */
+	bool bDeleteOrphanedVertexInstances;
+
+	/** Whether we should also delete any sections that are left empty after deleting this polygon */
+	bool bDeleteEmptySections;
+
+	/** Default constructor */
+	FDeletePolygonsChangeInput()
+		: PolygonIDsToDelete(),
+		  bDeleteOrphanedEdges( true ),
+		  bDeleteOrphanedVertices( true ),
+		  bDeleteOrphanedVertexInstances( true ),
+		  bDeleteEmptySections( false )
+	{
+	}
+};
+
+
+class FDeletePolygonsChange : public FChange
+{
+
+public:
+
+	/** Constructor */
+	FDeletePolygonsChange( const FDeletePolygonsChangeInput& InitInput )
+		: Input( InitInput )
+	{
+	}
+
+ 	FDeletePolygonsChange( FDeletePolygonsChangeInput&& InitInput )
+ 		: Input( MoveTemp ( InitInput ) )
+ 	{
+ 	}
+
+	// Parent class overrides
+	virtual TUniquePtr<FChange> Execute( UObject* Object ) override;
+	virtual FString ToString() const override;
+
+
+private:
+
+	/** The data we need to make this change */
+	FDeletePolygonsChangeInput Input;
+
+};
+
+
+struct FSetVerticesAttributesChangeInput
+{
+	TArray<FAttributesForVertex> AttributesForVertices;
+};
+
+
+class FSetVerticesAttributesChange : public FChange
+{
+
+public:
+
+	/** Constructor */
+	FSetVerticesAttributesChange( const FSetVerticesAttributesChangeInput& InitInput )
+		: Input( InitInput )
+	{
+	}
+
+ 	FSetVerticesAttributesChange( FSetVerticesAttributesChangeInput&& InitInput )
+ 		: Input( MoveTemp( InitInput ) )
+ 	{
+ 	}
+
+	// Parent class overrides
+	virtual TUniquePtr<FChange> Execute( UObject* Object ) override;
+	virtual FString ToString() const override;
+
+
+private:
+
+	/** The data we need to make this change */
+	FSetVerticesAttributesChangeInput Input;
+
+};
+
+
+struct FSetVertexInstancesAttributesChangeInput
+{
+	TArray<FAttributesForVertexInstance> AttributesForVertexInstances;
+};
+
+
+class FSetVertexInstancesAttributesChange : public FChange
+{
+
+public:
+
+	/** Constructor */
+	FSetVertexInstancesAttributesChange( const FSetVertexInstancesAttributesChangeInput& InitInput )
+		: Input( InitInput )
+	{
+	}
+
+	FSetVertexInstancesAttributesChange( FSetVertexInstancesAttributesChangeInput&& InitInput )
+		: Input( MoveTemp( InitInput ) )
+	{
+	}
+
+	// Parent class overrides
+	virtual TUniquePtr<FChange> Execute( UObject* Object ) override;
+	virtual FString ToString() const override;
+
+
+private:
+
+	/** The data we need to make this change */
+	FSetVertexInstancesAttributesChangeInput Input;
+
+};
+
+
+struct FSetEdgesAttributesChangeInput
+{
+	/** Which edges we'll be setting attributes for, along with the attribute data to set */
+	TArray<FAttributesForEdge> AttributesForEdges;
+
+	/** Default constructor */
+	FSetEdgesAttributesChangeInput()
+		: AttributesForEdges()
+	{
+	}
+};
+
+
+class FSetEdgesAttributesChange : public FChange
+{
+
+public:
+
+	/** Constructor */
+	FSetEdgesAttributesChange( const FSetEdgesAttributesChangeInput& InitInput )
+		: Input( InitInput )
+	{
+	}
+
+ 	FSetEdgesAttributesChange( FSetEdgesAttributesChangeInput&& InitInput )
+ 		: Input( MoveTemp( InitInput ) )
+ 	{
+ 	}
+
+	// Parent class overrides
+	virtual TUniquePtr<FChange> Execute( UObject* Object ) override;
+	virtual FString ToString() const override;
+
+
+private:
+
+	/** The data we need to make this change */
+	FSetEdgesAttributesChangeInput Input;
+
+};
+
+
+struct FSetPolygonsVertexAttributesChangeInput
+{
+	/** Which polygons we'll be setting vertex attributes for, along with the attribute data to set */
+	TArray<FVertexAttributesForPolygon> VertexAttributesForPolygons;
+
+	/** Default constructor */
+	FSetPolygonsVertexAttributesChangeInput()
+		: VertexAttributesForPolygons()
+	{
+	}
+};
+
+
+class FSetPolygonsVertexAttributesChange : public FChange
+{
+
+public:
+
+	/** Constructor */
+	FSetPolygonsVertexAttributesChange( const FSetPolygonsVertexAttributesChangeInput& InitInput )
+		: Input( InitInput )
+	{
+	}
+
+ 	FSetPolygonsVertexAttributesChange( FSetPolygonsVertexAttributesChangeInput&& InitInput )
+ 		: Input( MoveTemp( InitInput ) )
+ 	{
+ 	}
+
+	// Parent class overrides
+	virtual TUniquePtr<FChange> Execute( UObject* Object ) override;
+	virtual FString ToString() const override;
+
+
+private:
+
+	/** The data we need to make this change */
+	FSetPolygonsVertexAttributesChangeInput Input;
+
+};
+
+
+struct FChangePolygonsVertexInstancesChangeInput
+{
+	/** Which polygons we'll be setting vertex instances for, along with the vertex instances to set */
+	TArray<FChangeVertexInstancesForPolygon> VertexInstancesForPolygons;
+
+	/** Default constructor */
+	FChangePolygonsVertexInstancesChangeInput()
+		: VertexInstancesForPolygons()
+	{
+	}
+};
+
+
+class FChangePolygonsVertexInstancesChange : public FChange
+{
+
+public:
+
+	/** Constructor */
+	FChangePolygonsVertexInstancesChange( const FChangePolygonsVertexInstancesChangeInput& InitInput )
+		: Input( InitInput )
+	{
+	}
+
+	FChangePolygonsVertexInstancesChange( FChangePolygonsVertexInstancesChangeInput&& InitInput )
+		: Input( MoveTemp( InitInput ) )
+	{
+	}
+
+	// Parent class overrides
+	virtual TUniquePtr<FChange> Execute( UObject* Object ) override;
+	virtual FString ToString() const override;
+
+
+private:
+
+	/** The data we need to make this change */
+	FChangePolygonsVertexInstancesChangeInput Input;
+
+};
+
+
+struct FSetEdgesVerticesChangeInput
+{
+	/** The edge to set new vertices for */
+	TArray<FVerticesForEdge> VerticesForEdges;
+
+	/** Default constructor */
+	FSetEdgesVerticesChangeInput()
+		: VerticesForEdges()
+	{
+	}
+};
+
+
+class FSetEdgesVerticesChange : public FChange
+{
+
+public:
+
+	/** Constructor */
+	FSetEdgesVerticesChange( const FSetEdgesVerticesChangeInput& InitInput )
+		: Input( InitInput )
+	{
+	}
+
+ 	FSetEdgesVerticesChange( FSetEdgesVerticesChangeInput&& InitInput )
+ 		: Input( MoveTemp( InitInput ) )
+ 	{
+ 	}
+
+	// Parent class overrides
+	virtual TUniquePtr<FChange> Execute( UObject* Object ) override;
+	virtual FString ToString() const override;
+
+
+private:
+
+	/** The data we need to make this change */
+	FSetEdgesVerticesChangeInput Input;
+
+};
+
+
+struct FInsertPolygonPerimeterVerticesChangeInput
+{
+	/** The polygon we'll be inserting vertices into */
+	FPolygonID PolygonID;
+
+	/** The first polygon perimeter vertex number to insert indices before */
+	uint32 InsertBeforeVertexNumber;
+
+	/** The vertices to insert, along with their polygon perimeter vertex attributes */
+	TArray<FVertexAndAttributes> VerticesToInsert;
+
+	
+	/** Default constructor */
+	FInsertPolygonPerimeterVerticesChangeInput()
+		: PolygonID( FPolygonID::Invalid ),
+		  InsertBeforeVertexNumber( 0 ),
+		  VerticesToInsert()
+	{
+	}
+};
+
+
+class FInsertPolygonPerimeterVerticesChange : public FChange
+{
+
+public:
+
+	/** Constructor */
+	FInsertPolygonPerimeterVerticesChange( const FInsertPolygonPerimeterVerticesChangeInput& InitInput )
+		: Input( InitInput )
+	{
+	}
+
+ 	FInsertPolygonPerimeterVerticesChange( FInsertPolygonPerimeterVerticesChangeInput&& InitInput )
+ 		: Input( MoveTemp( InitInput ) )
+ 	{
+ 	}
+
+	// Parent class overrides
+	virtual TUniquePtr<FChange> Execute( UObject* Object ) override;
+	virtual FString ToString() const override;
+
+
+private:
+
+	/** The data we need to make this change */
+	FInsertPolygonPerimeterVerticesChangeInput Input;
+
+};
+
+
+struct FRemovePolygonPerimeterVerticesChangeInput
+{
+	/** The polygon we'll be removing vertices from */
+	FPolygonID PolygonID;
+
+	/** The first polygon perimeter vertex number to remove */
+	uint32 FirstVertexNumberToRemove;
+
+	/** The number of vertices to remove */
+	uint32 NumVerticesToRemove;
+
+	/** Whether orphaned vertex instances should be deleted or not */
+	bool bDeleteOrphanedVertexInstances;
+
+	
+	/** Default constructor */
+	FRemovePolygonPerimeterVerticesChangeInput()
+		: PolygonID( FPolygonID::Invalid ),
+		  FirstVertexNumberToRemove( 0 ),
+		  NumVerticesToRemove( 0 ),
+		  bDeleteOrphanedVertexInstances( false )
+	{
+	}
+};
+
+
+class FRemovePolygonPerimeterVerticesChange : public FChange
+{
+
+public:
+
+	/** Constructor */
+	FRemovePolygonPerimeterVerticesChange( const FRemovePolygonPerimeterVerticesChangeInput& InitInput )
+		: Input( InitInput )
+	{
+	}
+
+	// Parent class overrides
+	virtual TUniquePtr<FChange> Execute( UObject* Object ) override;
+	virtual FString ToString() const override;
+
+
+private:
+
+	/** The data we need to make this change */
+	FRemovePolygonPerimeterVerticesChangeInput Input;
+
+};
+
+
+struct FStartOrEndModificationChangeInput
+{
+	/** True if we should start modifying the mesh, or false if we should end modifying the mesh.  This will be reversed every time an undo/redo happens */
+	bool bStartModification;
+
+	/** The type of modification we're doing here */
+	EMeshModificationType MeshModificationType;
+	
+	/** Whether the mesh's topology can change during this modification */
+	EMeshTopologyChange MeshTopologyChange;
+
+	
+	/** Default constructor */
+	FStartOrEndModificationChangeInput()
+		: bStartModification( true ),
+		  MeshModificationType( EMeshModificationType::Final ),
+		  MeshTopologyChange( EMeshTopologyChange::TopologyChange )
+	{
+	}
+};
+
+
+class FStartOrEndModificationChange : public FChange
+{
+
+public:
+
+	/** Constructor */
+	FStartOrEndModificationChange( const FStartOrEndModificationChangeInput& InitInput )
+		: Input( InitInput )
+	{
+	}
+
+	// Parent class overrides
+	virtual TUniquePtr<FChange> Execute( UObject* Object ) override;
+	virtual FString ToString() const override;
+
+
+private:
+
+	/** The data we need to make this change */
+	FStartOrEndModificationChangeInput Input;
+
+};
+
+
+struct FSetSubdivisionCountChangeInput
+{
+	/** The new number of subdivisions to use */
+	int32 NewSubdivisionCount;
+
+	/** Default constructor */
+	FSetSubdivisionCountChangeInput()
+		: NewSubdivisionCount( 0 )
+	{
+	}
+};
+
+
+class FSetSubdivisionCountChange : public FChange
+{
+
+public:
+
+	/** Constructor */
+	FSetSubdivisionCountChange( const FSetSubdivisionCountChangeInput& InitInput )
+		: Input( InitInput )
+	{
+	}
+
+	// Parent class overrides
+	virtual TUniquePtr<FChange> Execute( UObject* Object ) override;
+	virtual FString ToString() const override;
+
+
+private:
+
+	/** The data we need to make this change */
+	FSetSubdivisionCountChangeInput Input;
+
+};
+
+
+struct FCreatePolygonGroupsChangeInput
+{
+	/** Information about the polygon groups to create */
+	TArray<FPolygonGroupToCreate> PolygonGroupsToCreate;
+
+	/** Default constructor */
+	FCreatePolygonGroupsChangeInput()
+		: PolygonGroupsToCreate()
+	{
+	}
+};
+
+
+class FCreatePolygonGroupsChange : public FChange
+{
+
+public:
+
+	/** Constructor */
+	FCreatePolygonGroupsChange( const FCreatePolygonGroupsChangeInput& InitInput )
+		: Input( InitInput )
+	{
+	}
+
+	// Parent class overrides
+	virtual TUniquePtr<FChange> Execute( UObject* Object ) override;
+	virtual FString ToString() const override;
+
+
+private:
+
+	/** The data we need to make this change */
+	FCreatePolygonGroupsChangeInput Input;
+};
+
+
+struct FDeletePolygonGroupsChangeInput
+{
+	/** The polygon group IDs to delete */
+	TArray<FPolygonGroupID> PolygonGroupIDs;
+
+	/** Default constructor */
+	FDeletePolygonGroupsChangeInput()
+		: PolygonGroupIDs()
+	{
+	}
+};
+
+
+class FDeletePolygonGroupsChange : public FChange
+{
+
+public:
+
+	/** Constructor */
+	FDeletePolygonGroupsChange( const FDeletePolygonGroupsChangeInput& InitInput )
+		: Input( InitInput )
+	{
+	}
+
+	// Parent class overrides
+	virtual TUniquePtr<FChange> Execute( UObject* Object ) override;
+	virtual FString ToString() const override;
+
+
+private:
+
+	/** The data we need to make this change */
+	FDeletePolygonGroupsChangeInput Input;
+};
diff --git a/Engine/Plugins/Runtime/EditableMesh/Source/EditableMesh/EditableMeshCustomVersion.cpp b/Engine/Plugins/Runtime/EditableMesh/Source/EditableMesh/EditableMeshCustomVersion.cpp
new file mode 100644
--- /dev/null
+++ b/Engine/Plugins/Runtime/EditableMesh/Source/EditableMesh/EditableMeshCustomVersion.cpp
@@ -0,0 +1,9 @@
+// Copyright 1998-2016 Epic Games, Inc. All Rights Reserved.
+
+#include "EditableMeshCustomVersion.h"
+#include "Serialization/CustomVersion.h"
+
+const FGuid FEditableMeshCustomVersion::GUID( 0xFB26E412, 0x1F154B4D, 0x9372550A, 0x961D2F70 );
+
+// Register the custom version with core
+FCustomVersionRegistration GRegisterEditableMeshCustomVersion( FEditableMeshCustomVersion::GUID, FEditableMeshCustomVersion::LatestVersion, TEXT( "EditableMeshVer" ) );
diff --git a/Engine/Plugins/Runtime/EditableMesh/Source/EditableMesh/EditableMeshCustomVersion.h b/Engine/Plugins/Runtime/EditableMesh/Source/EditableMesh/EditableMeshCustomVersion.h
new file mode 100644
--- /dev/null
+++ b/Engine/Plugins/Runtime/EditableMesh/Source/EditableMesh/EditableMeshCustomVersion.h
@@ -0,0 +1,34 @@
+// Copyright 1998-2016 Epic Games, Inc. All Rights Reserved.
+
+#pragma once
+
+#include "CoreMinimal.h"
+#include "Misc/Guid.h"
+
+// Custom serialization version for all packages containing EditableMesh asset types
+struct FEditableMeshCustomVersion
+{
+	enum Type
+	{
+		// Before any version changes were made in the module
+		BeforeCustomVersionWasAdded = 0,
+
+		// Added subdivision level count and texture coordinate count to UEditableMesh format
+		TextureCoordinateAndSubdivisionCounts,
+
+		// SparseArrays now have custom serialization, preserving the gaps
+		CustomSparseArraySerialization,
+
+		// Added vertex color support
+		WithVertexColors,
+
+		// -----<new versions can be added above this line>-------------------------------------------------
+		VersionPlusOne,
+		LatestVersion = VersionPlusOne - 1
+	};
+
+	// The GUID for this custom version number
+	const static FGuid GUID;
+
+	FEditableMeshCustomVersion() = delete;
+};
diff --git a/Engine/Plugins/Runtime/EditableMesh/Source/EditableMesh/EditableMeshFactory.cpp b/Engine/Plugins/Runtime/EditableMesh/Source/EditableMesh/EditableMeshFactory.cpp
new file mode 100644
--- /dev/null
+++ b/Engine/Plugins/Runtime/EditableMesh/Source/EditableMesh/EditableMeshFactory.cpp
@@ -0,0 +1,81 @@
+// Copyright 1998-2015 Epic Games, Inc. All Rights Reserved.
+
+#include "EditableMeshFactory.h"
+#include "Features/IModularFeatures.h"
+#include "IEditableMeshFormat.h"
+
+
+FEditableMeshSubMeshAddress UEditableMeshFactory::MakeSubmeshAddress( UPrimitiveComponent* PrimitiveComponent, const int32 LODIndex )
+{
+	check( PrimitiveComponent != nullptr );
+
+	FEditableMeshSubMeshAddress SubMeshAddress;
+
+	const int32 NumEditableMeshFormats = IModularFeatures::Get().GetModularFeatureImplementationCount( "EditableMeshFormat" );
+	for( int32 EditableMeshFormatIndex = 0; EditableMeshFormatIndex < NumEditableMeshFormats; ++EditableMeshFormatIndex )
+	{
+		IEditableMeshFormat& EditableMeshFormat = *static_cast<IEditableMeshFormat*>( IModularFeatures::Get().GetModularFeatureImplementation( "EditableMeshFormat", EditableMeshFormatIndex ) );
+
+		SubMeshAddress = FEditableMeshSubMeshAddress();
+		SubMeshAddress.MeshObjectPtr = nullptr;	// This will be filled in below (FillMeshObjectPtr)
+		SubMeshAddress.EditableMeshFormat = &EditableMeshFormat;
+		SubMeshAddress.LODIndex = LODIndex;
+		EditableMeshFormat.FillMeshObjectPtr( *PrimitiveComponent, SubMeshAddress );	// @todo mesheditor: This stuff is a bit clunky, would like to refactor it
+		if( SubMeshAddress.MeshObjectPtr != nullptr )
+		{
+			break;
+		}
+		else
+		{
+			SubMeshAddress = FEditableMeshSubMeshAddress();
+		}
+	}
+
+	return SubMeshAddress;
+}
+
+
+FEditableMeshSubMeshAddress UEditableMeshFactory::MakeSubmeshAddress( UStaticMesh& StaticMesh, const int32 LODIndex )
+{
+	FEditableMeshSubMeshAddress SubMeshAddress;
+
+	const int32 NumEditableMeshFormats = IModularFeatures::Get().GetModularFeatureImplementationCount("EditableMeshFormat");
+	
+	if(NumEditableMeshFormats > 0)
+	{
+		IEditableMeshFormat& EditableMeshFormat = *static_cast<IEditableMeshFormat*>(IModularFeatures::Get().GetModularFeatureImplementation("EditableMeshFormat", 0));
+
+		SubMeshAddress = FEditableMeshSubMeshAddress();
+		SubMeshAddress.MeshObjectPtr = &StaticMesh;	
+		SubMeshAddress.EditableMeshFormat = &EditableMeshFormat;
+		SubMeshAddress.LODIndex = LODIndex;
+	}
+
+	return SubMeshAddress;
+}
+
+UEditableMesh* UEditableMeshFactory::MakeEditableMesh( UPrimitiveComponent* PrimitiveComponent, const int32 LODIndex )
+{
+	check( PrimitiveComponent != nullptr );
+
+	FEditableMeshSubMeshAddress SubMeshAddress = MakeSubmeshAddress( PrimitiveComponent, LODIndex );
+	return MakeEditableMesh( PrimitiveComponent, SubMeshAddress );
+}
+
+
+UEditableMesh* UEditableMeshFactory::MakeEditableMesh( UPrimitiveComponent* PrimitiveComponent, const FEditableMeshSubMeshAddress& SubMeshAddress )
+{
+	check( PrimitiveComponent != nullptr );
+	UEditableMesh* EditableMesh = nullptr;
+
+	if( SubMeshAddress.EditableMeshFormat != nullptr &&
+		SubMeshAddress.MeshObjectPtr != nullptr )
+	{
+		// @todo mesheditor perf: This is going to HITCH
+		EditableMesh = SubMeshAddress.EditableMeshFormat->MakeEditableMesh( *PrimitiveComponent, SubMeshAddress );
+	}
+
+	return EditableMesh;
+}
+
+
diff --git a/Engine/Plugins/Runtime/EditableMesh/Source/EditableMesh/EditableMeshModule.cpp b/Engine/Plugins/Runtime/EditableMesh/Source/EditableMesh/EditableMeshModule.cpp
new file mode 100644
--- /dev/null
+++ b/Engine/Plugins/Runtime/EditableMesh/Source/EditableMesh/EditableMeshModule.cpp
@@ -0,0 +1,41 @@
+// Copyright 1998-2015 Epic Games, Inc. All Rights Reserved.
+
+#include "EditableMeshModule.h"
+#include "IEditableMeshModule.h"
+#include "StaticMeshEditableMeshFormat.h"
+
+
+class FEditableMeshModule : public IEditableMeshModule
+{
+public:
+
+	FEditableMeshModule()
+	{
+	}
+
+	virtual void StartupModule() override;
+
+	virtual void ShutdownModule() override;
+
+private:
+
+	/** Static mesh editing */
+	FStaticMeshEditableMeshFormat StaticMeshEditableMeshFormat;
+
+};
+
+
+void FEditableMeshModule::StartupModule()
+{
+	IModularFeatures::Get().RegisterModularFeature( "EditableMeshFormat", &StaticMeshEditableMeshFormat );
+}
+
+
+void FEditableMeshModule::ShutdownModule()
+{
+	IModularFeatures::Get().UnregisterModularFeature( "EditableMeshFormat", &StaticMeshEditableMeshFormat );
+}
+
+
+
+IMPLEMENT_MODULE( FEditableMeshModule, EditableMesh )
diff --git a/Engine/Plugins/Runtime/EditableMesh/Source/EditableMesh/EditableMeshModule.h b/Engine/Plugins/Runtime/EditableMesh/Source/EditableMesh/EditableMeshModule.h
new file mode 100644
--- /dev/null
+++ b/Engine/Plugins/Runtime/EditableMesh/Source/EditableMesh/EditableMeshModule.h
@@ -0,0 +1,8 @@
+// Copyright 1998-2015 Epic Games, Inc. All Rights Reserved.
+
+#pragma once
+
+#include "CoreMinimal.h"
+#include "IEditableMeshModule.h"
+
+
