Mesh Editor: New "EditableMesh" plugin (replaces "MeshEditingRuntime" module)

diff --git a/Engine/Plugins/Runtime/EditableMesh/Source/EditableMesh/Public/EditableMesh.h b/Engine/Plugins/Runtime/EditableMesh/Source/EditableMesh/Public/EditableMesh.h
new file mode 100644
--- /dev/null
+++ b/Engine/Plugins/Runtime/EditableMesh/Source/EditableMesh/Public/EditableMesh.h
@@ -0,0 +1,877 @@
+// Copyright 1998-2015 Epic Games, Inc. All Rights Reserved.
+
+#pragma once
+
+#include "CoreMinimal.h"
+#include "UObject/Object.h"
+#include "EditableMeshTypes.h"
+#include "EditableMeshCustomVersion.h"
+#include "Change.h"		// For TUniquePtr<FChange>
+#include "Logging/LogMacros.h"
+#include "Materials/MaterialInterface.h"
+#include "EditableMesh.generated.h"
+
+class UEditableMeshAdapter;
+
+EDITABLEMESH_API DECLARE_LOG_CATEGORY_EXTERN( LogEditableMesh, Log, All );
+
+
+// @todo mesheditor: Comment these classes and enums!
+
+namespace OpenSubdiv
+{
+	namespace v3_2_0
+	{
+		namespace Far
+		{
+			class TopologyRefiner;
+		}
+	}
+}
+
+
+/**
+* Perform custom serialization for TSparseArray.
+* The default TSparseArray serialization also compacts all the elements, removing the gaps and changing the indices.
+* The indices are significant in editable meshes, hence this is a custom serializer which preserves them.
+*/
+template <typename T>
+void SerializeSparseArray( FArchive& Ar, TSparseArray<T>& Array )
+{
+	if( Ar.CustomVer( FEditableMeshCustomVersion::GUID ) < FEditableMeshCustomVersion::CustomSparseArraySerialization )
+	{
+		Ar << Array;
+	}
+	else
+	{
+		Array.CountBytes( Ar );
+
+		if( Ar.IsLoading() )
+		{
+			// Load array
+			TBitArray<> AllocatedIndices;
+			Ar << AllocatedIndices;
+
+			Array.Empty( AllocatedIndices.Num() );
+			for( auto It = TConstSetBitIterator<>( AllocatedIndices ); It; ++It )
+			{
+				Array.Insert( It.GetIndex(), T() );
+				Ar << Array[ It.GetIndex() ];
+			}
+		}
+		else
+		{
+			// Save array
+			const int32 MaxIndex = Array.GetMaxIndex();
+
+			// We have to build the TBitArray representing allocated indices by hand, as we don't have access to it from outside TSparseArray.
+			// @todo core: consider replacing TSparseArray serialization with this format.
+			TBitArray<> AllocatedIndices( false, MaxIndex );
+			for( int32 Index = 0; Index < MaxIndex; ++Index )
+			{
+				if( Array.IsAllocated( Index ) )
+				{
+					AllocatedIndices[ Index ] = true;
+				}
+			}
+			Ar << AllocatedIndices;
+
+			for( auto It = Array.CreateIterator(); It; ++It )
+			{
+				Ar << *It;
+			}
+		}
+	}
+}
+
+
+template <typename T, typename ElementIDType>
+void CompactSparseArrayElements( TSparseArray<T>& Array, TSparseArray<ElementIDType>& IndexRemap )
+{
+	static_assert( TIsDerivedFrom<ElementIDType, FElementID>::IsDerived, "Remap array type must be derived from FElementID" );
+
+	static TSparseArray<T> NewArray;
+	NewArray.Empty( Array.Num() );
+
+	IndexRemap.Empty( Array.GetMaxIndex() );
+
+	// Add valid elements into a new contiguous sparse array.  Note non-const iterator so we can move elements.
+	for( TSparseArray<T>::TIterator It( Array ); It; ++It )
+	{
+		const int32 OldElementIndex = It.GetIndex();
+
+		// @todo mesheditor: implement TSparseArray::Add( ElementType&& ) to save this obscure approach
+		const int32 NewElementIndex = NewArray.Add( T() );
+		NewArray[ NewElementIndex ] = MoveTemp( *It );
+
+		// Provide an O(1) lookup from old index to new index, used when patching up vertex references afterwards
+		IndexRemap.Insert( OldElementIndex, ElementIDType( NewElementIndex ) );
+	}
+
+	Array = MoveTemp( NewArray );
+}
+
+
+template <typename T, typename ElementIDType>
+void RemapSparseArrayElements( TSparseArray<T>& Array, const TSparseArray<ElementIDType>& IndexRemap )
+{
+	static_assert( TIsDerivedFrom<ElementIDType, FElementID>::IsDerived, "Remap array type must be derived from FElementID" );
+
+	static TSparseArray<T> NewArray;
+	NewArray.Empty( IndexRemap.GetMaxIndex() );
+
+	// Add valid elements into a new contiguous sparse array.  Note non-const iterator so we can move elements.
+	for( TSparseArray<T>::TIterator It( Array ); It; ++It )
+	{
+		const int32 OldElementIndex = It.GetIndex();
+
+		check( IndexRemap.IsAllocated( OldElementIndex ) );
+		const int32 NewElementIndex = IndexRemap[ OldElementIndex ].GetValue();
+
+		// @todo mesheditor: implement TSparseArray::Insert( ElementType&& ) to save this obscure approach
+		NewArray.Insert( NewElementIndex, T() );
+		NewArray[ NewElementIndex ] = MoveTemp( *It );
+	}
+
+	Array = MoveTemp( NewArray );
+}
+
+
+USTRUCT()
+struct FMeshVertex
+{
+	GENERATED_BODY()
+
+	/** Position of the vertex. */
+	UPROPERTY()
+	FVector VertexPosition;
+
+	/** All of vertex instances which reference this vertex (for split vertex support) */
+	UPROPERTY()
+	TArray<FVertexInstanceID> VertexInstanceIDs;
+
+	/** The edges connected to this vertex */
+	UPROPERTY()
+	TArray<FEdgeID> ConnectedEdgeIDs;
+
+	/** When subdivisions are enabled, this controls how sharp the vertex is, between 0.0 and 1.0. */
+	UPROPERTY()
+	float CornerSharpness;	// @todo mesheditor subdiv: Not really used by static meshes at all.  Only for UEditableMeshes that use subdivision features.  Move elsewhere.?
+
+	/** Serializer */
+	friend FArchive& operator<<( FArchive& Ar, FMeshVertex& Vertex )
+	{
+		Ar << Vertex.VertexPosition;
+		Ar << Vertex.VertexInstanceIDs;
+		Ar << Vertex.ConnectedEdgeIDs;
+		Ar << Vertex.CornerSharpness;
+		return Ar;
+	}
+};
+
+
+USTRUCT()
+struct FMeshVertexInstance
+{
+	GENERATED_BODY()
+
+	/** The vertex this is instancing */
+	UPROPERTY()
+	FVertexID VertexID;
+
+	/** List of connected polygons */
+	UPROPERTY()
+	TArray<FPolygonID> ConnectedPolygons;
+
+	/** UVs for the vertex instance */
+	UPROPERTY()
+	TArray<FVector2D> VertexUVs;
+
+	/** Normal vector */
+	UPROPERTY()
+	FVector Normal;
+	
+	/** Tangent vector */
+	UPROPERTY()
+	FVector Tangent;
+
+	/** Basis determinant sign used to calculate the sense of the binormal */
+	UPROPERTY()
+	float BinormalSign;
+
+	/** Vertex color */
+	UPROPERTY()
+	FLinearColor Color;
+
+	/** Serializer */
+	friend FArchive& operator<<( FArchive& Ar, FMeshVertexInstance& VertexInstance )
+	{
+		Ar << VertexInstance.VertexID;
+		Ar << VertexInstance.ConnectedPolygons;
+		Ar << VertexInstance.VertexUVs;
+		Ar << VertexInstance.Normal;
+		Ar << VertexInstance.Tangent;
+		Ar << VertexInstance.BinormalSign;
+		Ar << VertexInstance.Color;
+		return Ar;
+	}
+};
+
+
+USTRUCT()
+struct FMeshEdge
+{
+	GENERATED_BODY()
+
+	/** IDs of the two editable mesh vertices that make up this edge.  The winding direction is not defined. */
+	UPROPERTY()
+	FVertexID VertexIDs[ 2 ];
+
+	/** The polygons that share this edge.  It's best if there are always only two polygons that share
+	    the edge, and those polygons are facing the same direction */
+	UPROPERTY()
+	TArray<FPolygonID> ConnectedPolygons;
+
+	/** Whether this edge is 'hard' or not, for the purpose of vertex normal and tangent generation */
+	UPROPERTY()
+	bool bIsHardEdge;
+
+	/** When subdivisions are enabled, this controls how sharp the creasing of this edge will be, between 0.0 and 1.0. */
+	UPROPERTY()
+	float CreaseSharpness;
+
+	/** Serializer */
+	friend FArchive& operator<<( FArchive& Ar, FMeshEdge& Edge )
+	{
+		Ar << Edge.VertexIDs[ 0 ];
+		Ar << Edge.VertexIDs[ 1 ];
+		Ar << Edge.ConnectedPolygons;
+		Ar << Edge.bIsHardEdge;
+		Ar << Edge.CreaseSharpness;
+		return Ar;
+	}
+};
+
+
+USTRUCT()
+struct FMeshPolygonContour
+{
+	GENERATED_BODY()
+
+	/** The ordered list of vertex instances which make up the polygon contour. The winding direction is counter-clockwise. */
+	UPROPERTY()
+	TArray<FVertexInstanceID> VertexInstanceIDs;
+
+	/** Serializer */
+	friend FArchive& operator<<( FArchive& Ar, FMeshPolygonContour& Contour )
+	{
+		Ar << Contour.VertexInstanceIDs;
+		return Ar;
+	}
+};
+
+
+USTRUCT( BlueprintType )
+struct FMeshTriangle
+{
+	GENERATED_BODY()
+
+	/** First vertex instance that makes up this triangle.  Indices must be ordered counter-clockwise. */
+	UPROPERTY( BlueprintReadWrite, Category="Editable Mesh" )
+	FVertexInstanceID VertexInstanceID0;
+
+	/** Second vertex instance that makes up this triangle.  Indices must be ordered counter-clockwise. */
+	UPROPERTY( BlueprintReadWrite, Category="Editable Mesh" )
+	FVertexInstanceID VertexInstanceID1;
+
+	/** Third vertex instance that makes up this triangle.  Indices must be ordered counter-clockwise. */
+	UPROPERTY( BlueprintReadWrite, Category="Editable Mesh" )
+	FVertexInstanceID VertexInstanceID2;
+
+	/** Gets the specified triangle vertex instance ID.  Pass an index between 0 and 2 inclusive. */
+	inline FVertexInstanceID GetVertexInstanceID( const int32 Index ) const
+	{
+		checkSlow( Index >= 0 && Index <= 2 );
+		return reinterpret_cast<const FVertexInstanceID*>( this )[ Index ];
+	}
+
+	/** Sets the specified triangle vertex instance ID.  Pass an index between 0 and 2 inclusive, and the new vertex instance ID to store. */
+	inline void SetVertexInstanceID( const int32 Index, const FVertexInstanceID NewVertexInstanceID )
+	{
+		checkSlow( Index >= 0 && Index <= 2 );
+		( reinterpret_cast<FVertexInstanceID*>( this )[ Index ] ) = NewVertexInstanceID;
+	}
+
+	/** Serializer */
+	friend FArchive& operator<<( FArchive& Ar, FMeshTriangle& Tri )
+	{
+		Ar << Tri.VertexInstanceID0;
+		Ar << Tri.VertexInstanceID1;
+		Ar << Tri.VertexInstanceID2;
+		return Ar;
+	}
+};
+
+
+USTRUCT()
+struct FMeshPolygon
+{
+	GENERATED_BODY()
+
+	/** The outer boundary edges of this polygon */
+	UPROPERTY()
+	FMeshPolygonContour PerimeterContour;
+
+	/** Optional inner contours of this polygon that define holes inside of the polygon.  For the geometry to
+	    be considered valid, the hole contours should reside within the boundary of the polygon perimeter contour, 
+		and must not overlap each other.  No "nesting" of polygons inside the holes is supported -- those are 
+		simply separate polygons */
+	UPROPERTY()
+	TArray<FMeshPolygonContour> HoleContours;
+
+	/** List of triangles which make up this polygon */
+	UPROPERTY()
+	TArray<FMeshTriangle> Triangles;
+
+	/** The polygon group which contains this polygon */
+	UPROPERTY()
+	FPolygonGroupID PolygonGroupID;
+
+	/** Cached normal */
+	UPROPERTY()
+	FVector PolygonNormal;
+
+	/** Cached tangent */
+	UPROPERTY()
+	FVector PolygonTangent;
+
+	/** Cached binormal */
+	UPROPERTY()
+	FVector PolygonBinormal;
+
+	/** Cached center */
+	UPROPERTY()
+	FVector PolygonCenter;
+
+	/** Serializer */
+	friend FArchive& operator<<( FArchive& Ar, FMeshPolygon& Polygon )
+	{
+		Ar << Polygon.PerimeterContour;
+		Ar << Polygon.HoleContours;
+		Ar << Polygon.Triangles;
+		Ar << Polygon.PolygonGroupID;
+		return Ar;
+	}
+};
+
+
+USTRUCT()
+struct FMeshPolygonGroup
+{
+	GENERATED_BODY()
+
+	/** The material for this mesh section */
+	UPROPERTY()
+	UMaterialInterface* Material;
+
+	/** If true, collision is enabled for this section. */
+	UPROPERTY()
+	bool bEnableCollision;
+
+	/** If true, this section will cast a shadow */
+	UPROPERTY()
+	bool bCastShadow;
+
+	/** All polygons in this group */
+	TArray<FPolygonID> Polygons;
+
+	/** Serializer */
+	friend FArchive& operator<<( FArchive& Ar, FMeshPolygonGroup& PolygonGroup )
+	{
+		Ar << PolygonGroup.Material;
+		Ar << PolygonGroup.bEnableCollision;
+		Ar << PolygonGroup.bCastShadow;
+		Ar << PolygonGroup.Polygons;
+		return Ar;
+	}
+};
+
+
+UENUM( BlueprintType )
+enum class EInsetPolygonsMode : uint8
+{
+	All,
+	CenterPolygonOnly,
+	SidePolygonsOnly,
+};
+
+
+UENUM( BlueprintType )
+enum class ETriangleTessellationMode : uint8
+{
+	/** Connect each vertex to a new center vertex, forming three triangles */
+	ThreeTriangles,
+
+	/** Split each edge and create a center polygon that connects those new vertices, then three additional polygons for each original corner */
+	FourTriangles,
+};
+
+
+struct FElementIDRemappings
+{
+	TSparseArray<FVertexID> NewVertexIndexLookup;
+	TSparseArray<FVertexInstanceID> NewVertexInstanceIndexLookup;
+	TSparseArray<FEdgeID> NewEdgeIndexLookup;
+	TSparseArray<FPolygonID> NewPolygonIndexLookup;
+	TSparseArray<FPolygonGroupID> NewPolygonGroupIndexLookup;
+
+	FVertexID GetRemappedVertexID( FVertexID VertexID ) const
+	{
+		checkSlow( NewVertexIndexLookup.IsAllocated( VertexID.GetValue() ) );
+		return NewVertexIndexLookup[ VertexID.GetValue() ];
+	}
+
+	FVertexInstanceID GetRemappedVertexInstanceID( FVertexInstanceID VertexInstanceID ) const
+	{
+		checkSlow( NewVertexInstanceIndexLookup.IsAllocated( VertexInstanceID.GetValue() ) );
+		return NewVertexInstanceIndexLookup[ VertexInstanceID.GetValue() ];
+	}
+
+	FEdgeID GetRemappedEdgeID( FEdgeID EdgeID ) const
+	{
+		checkSlow( NewEdgeIndexLookup.IsAllocated( EdgeID.GetValue() ) );
+		return NewEdgeIndexLookup[ EdgeID.GetValue() ];
+	}
+
+	FPolygonID GetRemappedPolygonID( FPolygonID PolygonID ) const
+	{
+		checkSlow( NewPolygonIndexLookup.IsAllocated( PolygonID.GetValue() ) );
+		return NewPolygonIndexLookup[ PolygonID.GetValue() ];
+	}
+
+	FPolygonGroupID GetRemappedPolygonGroupID( FPolygonGroupID PolygonGroupID ) const
+	{
+		check( NewPolygonGroupIndexLookup.IsAllocated( PolygonGroupID.GetValue() ) );
+		return NewPolygonGroupIndexLookup[ PolygonGroupID.GetValue() ];
+	}
+};
+
+
+UCLASS( BlueprintType )
+class EDITABLEMESH_API UEditableMesh : public UObject
+{
+	GENERATED_BODY()
+
+public:
+
+	/** Default constructor that initializes good defaults for UEditableMesh */
+	UEditableMesh();
+
+	// UObject interface
+	virtual void Serialize( FArchive& Ar ) override;
+	virtual void PostLoad() override;
+
+	/** Compacts mesh element arrays to remove gaps, and fixes up referenced IDs */
+	void Compact();
+
+	/** Remaps mesh element arrays according to the provided remappings, in order to undo a compact operation */
+	void Uncompact( const FElementIDRemappings& Remappings );
+
+	UFUNCTION( BlueprintCallable, Category="Editable Mesh" ) void RebuildRenderMesh();
+	UFUNCTION( BlueprintCallable, Category="Editable Mesh" ) void StartModification( const EMeshModificationType MeshModificationType, const EMeshTopologyChange MeshTopologyChange );
+	UFUNCTION( BlueprintCallable, Category="Editable Mesh" ) void EndModification( const bool bFromUndo = false );
+	UFUNCTION( BlueprintPure, Category="Editable Mesh" ) bool IsCommitted() const;
+	UFUNCTION( BlueprintPure, Category="Editable Mesh" ) bool IsCommittedAsInstance() const;
+	UFUNCTION( BlueprintCallable, Category="Editable Mesh" ) void Commit();
+	UFUNCTION( BlueprintCallable, Category="Editable Mesh" ) UEditableMesh* CommitInstance( UPrimitiveComponent* ComponentToInstanceTo );
+	UFUNCTION( BlueprintCallable, Category="Editable Mesh" ) void Revert();
+	UFUNCTION( BlueprintCallable, Category="Editable Mesh" ) UEditableMesh* RevertInstance();
+	UFUNCTION( BlueprintCallable, Category="Editable Mesh" ) void PropagateInstanceChanges();
+
+	UFUNCTION( BlueprintPure, Category="Editable Mesh" ) int32 GetVertexCount() const;
+	UFUNCTION( BlueprintPure, Category="Editable Mesh" ) int32 GetVertexArraySize() const;
+	UFUNCTION( BlueprintPure, Category="Editable Mesh" ) bool IsValidVertex( const FVertexID VertexID ) const;
+	UFUNCTION( BlueprintPure, Category="Editable Mesh" ) FVector4 GetVertexAttribute( const FVertexID VertexID, const FName AttributeName, const int32 AttributeIndex ) const;
+	UFUNCTION( BlueprintPure, Category="Editable Mesh" ) int32 GetVertexConnectedEdgeCount( const FVertexID VertexID ) const;
+	UFUNCTION( BlueprintPure, Category="Editable Mesh" ) FEdgeID GetVertexConnectedEdge( const FVertexID VertexID, const int32 ConnectedEdgeNumber ) const;
+
+	UFUNCTION( BlueprintPure, Category="Editable Mesh" ) int32 GetVertexInstanceCount() const;
+	UFUNCTION( BlueprintPure, Category="Editable Mesh" ) int32 GetVertexInstanceArraySize() const;
+	UFUNCTION( BlueprintPure, Category="Editable Mesh" ) FVector4 GetVertexInstanceAttribute( const FVertexInstanceID VertexInstanceID, const FName AttributeName, const int32 AttributeIndex ) const;
+
+	UFUNCTION( BlueprintPure, Category="Editable Mesh" ) FVector4 GetEdgeAttribute( const FEdgeID EdgeID, const FName AttributeName, const int32 AttributeIndex ) const;
+	UFUNCTION( BlueprintPure, Category="Editable Mesh" ) int32 GetEdgeCount() const;
+	UFUNCTION( BlueprintPure, Category="Editable Mesh" ) int32 GetEdgeArraySize() const;
+	UFUNCTION( BlueprintPure, Category="Editable Mesh" ) bool IsValidEdge( const FEdgeID EdgeID ) const;
+	UFUNCTION( BlueprintPure, Category="Editable Mesh" ) FVertexID GetEdgeVertex( const FEdgeID EdgeID, const int32 EdgeVertexNumber ) const;
+	UFUNCTION( BlueprintPure, Category="Editable Mesh" ) int32 GetEdgeConnectedPolygonCount( const FEdgeID EdgeID ) const;
+	UFUNCTION( BlueprintPure, Category="Editable Mesh" ) FPolygonID GetEdgeConnectedPolygon( const FEdgeID EdgeID, const int32 ConnectedPolygonNumber ) const;
+
+	UFUNCTION( BlueprintPure, Category="Editable Mesh" ) int32 GetPolygonGroupCount() const;
+	UFUNCTION( BlueprintPure, Category="Editable Mesh" ) int32 GetPolygonGroupArraySize() const;
+	UFUNCTION( BlueprintPure, Category="Editable Mesh" ) bool IsValidPolygonGroup( const FPolygonGroupID PolygonGroupID ) const;
+	UFUNCTION( BlueprintPure, Category="Editable Mesh" ) int32 GetPolygonCountInGroup( const FPolygonGroupID PolygonGroupID ) const;
+	UFUNCTION( BlueprintPure, Category="Editable Mesh" ) FPolygonID GetPolygonInGroup( const FPolygonGroupID PolygonGroupID, const int32 PolygonNumber ) const;
+
+	UFUNCTION( BlueprintPure, Category="Editable Mesh" ) int32 GetPolygonCount() const;
+	UFUNCTION( BlueprintPure, Category="Editable Mesh" ) int32 GetPolygonArraySize() const;
+	UFUNCTION( BlueprintPure, Category="Editable Mesh" ) bool IsValidPolygon( const FPolygonID PolygonID ) const;
+	UFUNCTION( BlueprintPure, Category="Editable Mesh" ) FPolygonGroupID GetGroupForPolygon( const FPolygonID PolygonID ) const;
+
+//	UFUNCTION( BlueprintPure, Category="Editable Mesh" ) virtual int32 GetTriangleCount( const FPolygonGroupID PolygonGroupID ) const PURE_VIRTUAL(,return 0;);
+//	UFUNCTION( BlueprintPure, Category="Editable Mesh" ) virtual int32 GetTriangleArraySize( const FPolygonGroupID PolygonGroupID ) const PURE_VIRTUAL(,return 0;);
+
+	UFUNCTION( BlueprintPure, Category="Editable Mesh" ) int32 GetPolygonPerimeterVertexCount( const FPolygonID PolygonID ) const;
+	UFUNCTION( BlueprintPure, Category="Editable Mesh" ) FVertexID GetPolygonPerimeterVertex( const FPolygonID PolygonID, const int32 PolygonVertexNumber ) const;
+	UFUNCTION( BlueprintPure, Category="Editable Mesh" ) FVertexInstanceID GetPolygonPerimeterVertexInstance( const FPolygonID PolygonID, const int32 PolygonVertexNumber ) const;
+	UFUNCTION( BlueprintPure, Category="Editable Mesh" ) FVector4 GetPolygonPerimeterVertexAttribute( const FPolygonID PolygonID, const int32 PolygonVertexNumber, const FName AttributeName, const int32 AttributeIndex ) const;
+	UFUNCTION( BlueprintPure, Category="Editable Mesh" ) int32 GetPolygonHoleCount( const FPolygonID PolygonID ) const;
+	UFUNCTION( BlueprintPure, Category="Editable Mesh" ) int32 GetPolygonHoleVertexCount( const FPolygonID PolygonID, const int32 HoleNumber ) const;
+	UFUNCTION( BlueprintPure, Category="Editable Mesh" ) FVertexID GetPolygonHoleVertex( const FPolygonID PolygonID, const int32 HoleNumber, const int32 PolygonVertexNumber ) const;
+	UFUNCTION( BlueprintPure, Category="Editable Mesh" ) FVertexInstanceID GetPolygonHoleVertexInstance( const FPolygonID PolygonID, const int32 HoleNumber, const int32 PolygonVertexNumber ) const;
+	UFUNCTION( BlueprintPure, Category="Editable Mesh" ) FVector4 GetPolygonHoleVertexAttribute( const FPolygonID PolygonID, const int32 HoleNumber, const int32 PolygonVertexNumber, const FName AttributeName, const int32 AttributeIndex ) const;
+
+	UFUNCTION( BlueprintPure, Category="Editable Mesh" ) int32 GetPolygonTriangulatedTriangleCount( const FPolygonID PolygonID ) const;
+	UFUNCTION( BlueprintPure, Category="Editable Mesh" ) FMeshTriangle GetPolygonTriangulatedTriangle( const FPolygonID PolygonID, int32 PolygonTriangleNumber ) const;
+	UFUNCTION( BlueprintPure, Category="Editable Mesh" ) FVector GetPolygonTriangulatedTriangleVertexPosition( const FPolygonID PolygonID, const int32 PolygonTriangleNumber, const int32 TriangleVertexNumber ) const;
+
+protected:
+
+	/** Given a set of index remappings, fixes up references to element IDs */
+	void FixUpElementIDs( const FElementIDRemappings& Remappings );
+
+	FEdgeID GetPolygonContourEdge( const FMeshPolygonContour& Contour, const int32 ContourEdgeNumber, bool& bOutEdgeWindingIsReversedForPolygon ) const;
+	void GetPolygonContourEdges( const FMeshPolygonContour& Contour, TArray<FEdgeID>& OutPolygonContourEdgeIDs ) const;
+	void SetVertexAttribute( const FVertexID VertexID, const FName AttributeName, const int32 AttributeIndex, const FVector4 AttributeValue );
+	void SetVertexInstanceAttribute( const FVertexInstanceID VertexInstanceID, const FName AttributeName, const int32 AttributeIndex, const FVector4 AttributeValue );
+	void SetEdgeAttribute( const FEdgeID EdgeID, const FName AttributeName, const int32 AttributeIndex, const FVector4 AttributeValue, bool bIsUndo );
+	void GetVertexInstanceConnectedPolygonsInSameGroup( const FVertexInstanceID VertexInstanceID, const FPolygonID PolygonID, TArray<FPolygonID>& OutPolygonIDs ) const;
+	void SetEdgeHardness( const FEdgeID EdgeID, const bool bIsHard, const bool bIsUndo );
+	FVertexInstanceID CreateVertexInstanceForContourVertex( const FVertexAndAttributes& ContourVertex, const FPolygonID PolygonID );
+	void CreatePolygonContour( const FPolygonID PolygonID, const TArray<FVertexAndAttributes>& Contour, const EPolygonEdgeHardness PolygonEdgeHardness, TArray<FEdgeID>& OutEdgeIDs, TArray<FVertexInstanceID>& OutVertexInstanceIDs );
+	void BackupPolygonContour( const FMeshPolygonContour& Contour, TArray<FVertexAndAttributes>& OutVerticesAndAttributes );
+	void DeletePolygonContour( const FPolygonID PolygonID, const FMeshPolygonContour& Contour, const bool bDeleteOrphanedEdges, const bool bDeleteOrphanedVertexInstances, TArray<FEdgeID>& OrphanedEdgeIDs, TArray<FVertexInstanceID>& OrphanedVertexInstanceIDs );
+	static const TArray<FName>& GetMergeableVertexInstanceAttributes();
+	void GetMergedNamedVertexInstanceAttributeData( const TArray<FName>& AttributeNames, const FVertexInstanceID VertexInstanceID, const TArray<FMeshElementAttributeData>& AttributeValuesToMerge, TArray<FMeshElementAttributeData>& OutAttributes ) const;
+	void GetConnectedSoftEdges( const FVertexID VertexID, TArray<FEdgeID>& OutConnectedSoftEdges ) const;
+	void GetPolygonsInSameSoftEdgedGroup( const FVertexID VertexInstanceID, const FPolygonID PolygonID, TArray<FPolygonID>& OutPolygonIDs ) const;
+	void GetVertexInstancesInSameSoftEdgedGroup( const FVertexID VertexID, const FPolygonID PolygonID, const bool bPolygonNotYetInitialized, TArray<FVertexInstanceID>& OutVertexInstanceIDs ) const;
+	FVertexInstanceID GetVertexInstanceInPolygonForVertex( const FPolygonID PolygonID, const FVertexID VertexID ) const;
+	void SetPolygonContourVertexAttributes( FMeshPolygonContour& Contour, const FPolygonID PolygonID, const int32 HoleIndex, const TArray<FMeshElementAttributeList>& AttributeLists );
+	bool AreAttributeListsEqual( const TArray<FMeshElementAttributeData>& A, const TArray<FMeshElementAttributeData>& B ) const;
+	void SplitVertexInstanceInPolygons( const FVertexInstanceID VertexInstanceID, const TArray<FPolygonID>& PolygonIDs );
+	void ReplaceVertexInstanceInPolygons( const FVertexInstanceID OldVertexInstanceID, const FVertexInstanceID NewVertexInstanceID, const TArray<FPolygonID>& PolygonIDs );
+	float GetPolygonCornerAngleForVertex( const FPolygonID PolygonID, const FVertexID VertexID ) const;
+	FPolygonGroupID GetPolygonGroupIDFromMaterial( UMaterialInterface* Material, bool bCreateNewSectionIfNotFound );
+
+public:
+
+	/**
+	 * Called at initialization time to set this mesh's sub-mesh address
+	 *
+	 * @param	NewSubMeshAddress	The new sub-mesh address for this mesh
+	 */
+	void SetSubMeshAddress( const FEditableMeshSubMeshAddress& NewSubMeshAddress );
+
+	/**
+	 * @return	Returns true if StartModification() was called and the mesh is able to be modified currently.  Remember to call EndModification() when finished.
+	 */
+	UFUNCTION( BlueprintPure, Category="Editable Mesh" ) virtual bool IsBeingModified() const
+	{
+		return bIsBeingModified;
+	}
+
+	/**
+	 * @return	Returns true if undo tracking is enabled on this mesh
+	 */
+	UFUNCTION( BlueprintPure, Category="Editable Mesh" )
+	bool IsUndoAllowed() const
+	{
+		return bAllowUndo;
+	}
+
+	/**
+	 * Sets whether undo is allowed on this mesh
+	 *
+	 * @param	bInAllowUndo	True if undo features are enabled on this mesh.  You're only allowed to call MakeUndo() if this is set to true.
+	 */
+	UFUNCTION( BlueprintCallable, Category="Editable Mesh" )
+	void SetAllowUndo( const bool bInAllowUndo )
+	{
+		bAllowUndo = bInAllowUndo;
+	}
+
+
+	/**
+	 * @return	Returns true if there are any existing tracked changes that can be undo right now
+	 */
+	UFUNCTION( BlueprintPure, Category="Editable Mesh" )
+	bool AnyChangesToUndo() const;
+
+
+	/**
+	 * Gets the sub-mesh address for this mesh which uniquely identifies the mesh among other sub-meshes in the same component
+	 *
+	 * @return	The sub-mesh address for the mesh
+	 */
+	const FEditableMeshSubMeshAddress& GetSubMeshAddress() const;
+
+	/** Grabs any outstanding changes to this mesh and returns a change that can be used to undo those changes.  Calling this
+	    function will clear the history of changes.  This function will return a null pointer if bAllowUndo is false. */
+	// @todo mesheditor script: We might need this to be available for BP editable meshes, in some form at least.  Probably it should just apply the undo right away.
+	TUniquePtr<FChange> MakeUndo();
+
+
+	/**
+	 * Statics
+	 */
+
+	UFUNCTION( BlueprintPure, Category="Editable Mesh" ) static const TArray<FName>& GetValidVertexAttributes();
+	UFUNCTION( BlueprintPure, Category="Editable Mesh" ) static const TArray<FName>& GetValidVertexInstanceAttributes();
+	UFUNCTION( BlueprintPure, Category="Editable Mesh" ) static const TArray<FName>& GetValidEdgeAttributes();
+	UFUNCTION( BlueprintPure, Category="Editable Mesh" ) static const TArray<FName>& GetValidPolygonAttributes();
+	UFUNCTION( BlueprintPure, Category="Editable Mesh" ) static inline FVertexID InvalidVertexID()
+	{
+		return FVertexID::Invalid;
+	}
+	UFUNCTION( BlueprintPure, Category="Editable Mesh" ) static inline FEdgeID InvalidEdgeID()
+	{
+		return FEdgeID::Invalid;
+	}
+	UFUNCTION( BlueprintPure, Category="Editable Mesh" ) static inline FPolygonGroupID InvalidPolygonGroupID()
+	{
+		return FPolygonGroupID::Invalid;
+	}
+	UFUNCTION( BlueprintPure, Category="Editable Mesh" ) static inline FPolygonID InvalidPolygonID()
+	{
+		return FPolygonID::Invalid;
+	}
+	UFUNCTION( BlueprintPure, Category="Editable Mesh" ) static inline FVertexID MakeVertexID( const int32 VertexIndex )
+	{
+		return FVertexID( VertexIndex );
+	}
+	UFUNCTION( BlueprintPure, Category="Editable Mesh" ) static inline FEdgeID MakeEdgeID( const int32 EdgeIndex )
+	{
+		return FEdgeID( EdgeIndex );
+	}
+	UFUNCTION( BlueprintPure, Category="Editable Mesh" ) static inline FPolygonGroupID MakePolygonGroupID( const int32 PolygonGroupIndex )
+	{
+		return FPolygonGroupID( PolygonGroupIndex );
+	}
+	UFUNCTION( BlueprintPure, Category="Editable Mesh" ) static inline FPolygonID MakePolygonID( const int32 PolygonIndex )
+	{
+		return FPolygonID( PolygonIndex );
+	}
+
+	UFUNCTION( BlueprintPure, Category="Editable Mesh" ) int32 GetMaxAttributeIndex( const FName AttributeName ) const;
+	UFUNCTION( BlueprintPure, Category="Editable Mesh" ) FPolygonGroupID GetFirstValidPolygonGroup() const;
+	UFUNCTION( BlueprintPure, Category="Editable Mesh" ) int32 GetTextureCoordinateCount() const;
+	UFUNCTION( BlueprintPure, Category="Editable Mesh" ) int32 GetSubdivisionCount() const;
+	UFUNCTION( BlueprintPure, Category="Editable Mesh" ) bool IsPreviewingSubdivisions() const;
+	UFUNCTION( BlueprintPure, Category="Editable Mesh" ) void GetVertexConnectedEdges( const FVertexID VertexID, TArray<FEdgeID>& OutConnectedEdgeIDs ) const;
+	UFUNCTION( BlueprintPure, Category="Editable Mesh" ) void GetVertexConnectedPolygons( const FVertexID VertexID, TArray<FPolygonID>& OutConnectedPolygonIDs ) const;
+	UFUNCTION( BlueprintPure, Category="Editable Mesh" ) void GetVertexInstanceConnectedPolygons( const FVertexInstanceID VertexInstanceID, TArray<FPolygonID>& OutConnectedPolygonIDs ) const;
+	UFUNCTION( BlueprintPure, Category="Editable Mesh" ) int32 GetVertexInstanceConnectedPolygonCount( const FVertexInstanceID VertexInstanceID ) const;
+	UFUNCTION( BlueprintPure, Category="Editable Mesh" ) void GetVertexAdjacentVertices( const FVertexID VertexID, TArray< FVertexID >& OutAdjacentVertexIDs ) const;
+	UFUNCTION( BlueprintPure, Category="Editable Mesh" ) FEdgeID GetVertexPairEdge( const FVertexID VertexID, const FVertexID NextVertexID, bool& bOutEdgeWindingIsReversed ) const;
+	UFUNCTION( BlueprintPure, Category="Editable Mesh" ) void GetEdgeVertices( const FEdgeID EdgeID, FVertexID& OutEdgeVertexID0, FVertexID& OutEdgeVertexID1 ) const;
+	UFUNCTION( BlueprintPure, Category="Editable Mesh" ) void GetEdgeConnectedPolygons( const FEdgeID EdgeID, TArray<FPolygonID>& OutConnectedPolygonIDs ) const;
+	UFUNCTION( BlueprintPure, Category="Editable Mesh" ) void GetEdgeLoopElements( const FEdgeID EdgeID, TArray<FEdgeID>& EdgeLoopIDs ) const;
+	UFUNCTION( BlueprintPure, Category="Editable Mesh" ) FEdgeID GetEdgeThatConnectsVertices( const FVertexID VertexID0, const FVertexID VertexID1 ) const;
+	UFUNCTION( BlueprintPure, Category="Editable Mesh" ) int32 GetPolygonPerimeterEdgeCount( const FPolygonID PolygonID ) const;
+	UFUNCTION( BlueprintPure, Category="Editable Mesh" ) int32 GetPolygonHoleEdgeCount( const FPolygonID PolygonID, const int32 HoleNumber ) const;
+	UFUNCTION( BlueprintPure, Category="Editable Mesh" ) void GetPolygonPerimeterVertices( const FPolygonID PolygonID, TArray<FVertexID>& OutPolygonPerimeterVertexIDs ) const;
+	UFUNCTION( BlueprintPure, Category="Editable Mesh" ) void GetPolygonPerimeterVertexInstances( const FPolygonID PolygonID, TArray<FVertexInstanceID>& OutPolygonPerimeterVertexInstanceIDs ) const;
+	UFUNCTION( BlueprintPure, Category="Editable Mesh" ) void GetPolygonHoleVertices( const FPolygonID PolygonID, const int32 HoleNumber, TArray<FVertexID>& OutHoleVertexIDs ) const;
+	UFUNCTION( BlueprintPure, Category="Editable Mesh" ) void GetPolygonHoleVertexInstances( const FPolygonID PolygonID, const int32 HoleNumber, TArray<FVertexInstanceID>& OutHoleVertexInstanceIDs ) const;
+	UFUNCTION( BlueprintPure, Category="Editable Mesh" ) FEdgeID GetPolygonPerimeterEdge( const FPolygonID PolygonID, const int32 PerimeterEdgeNumber, bool& bOutEdgeWindingIsReversedForPolygon ) const;
+	UFUNCTION( BlueprintPure, Category="Editable Mesh" ) FEdgeID GetPolygonHoleEdge( const FPolygonID PolygonID, const int32 HoleNumber, const int32 HoleEdgeNumber ) const;
+	UFUNCTION( BlueprintPure, Category="Editable Mesh" ) void GetPolygonPerimeterEdges( const FPolygonID PolygonID, TArray<FEdgeID>& OutPolygonPerimeterEdgeIDs ) const;
+	UFUNCTION( BlueprintPure, Category="Editable Mesh" ) void GetPolygonHoleEdges( const FPolygonID PolygonID, const int32 HoleNumber, TArray<FEdgeID>& OutHoleEdgeIDs ) const;
+	UFUNCTION( BlueprintPure, Category="Editable Mesh" ) void GetPolygonAdjacentPolygons( const FPolygonID PolygonID, TArray<FPolygonID>& OutAdjacentPolygons ) const;
+	UFUNCTION( BlueprintPure, Category="Editable Mesh" ) int32 FindPolygonPerimeterVertexNumberForVertex( const FPolygonID PolygonID, const FVertexID VertexID ) const;
+	UFUNCTION( BlueprintPure, Category="Editable Mesh" ) int32 FindPolygonHoleVertexNumberForVertex( const FPolygonID PolygonID, const int32 HoleNumber, const FVertexID VertexID ) const;
+	UFUNCTION( BlueprintPure, Category="Editable Mesh" ) int32 FindPolygonPerimeterEdgeNumberForVertices( const FPolygonID PolygonID, const FVertexID EdgeVertexID0, const FVertexID EdgeVertexID1 ) const;
+	UFUNCTION( BlueprintPure, Category="Editable Mesh" ) int32 FindPolygonHoleEdgeNumberForVertices( const FPolygonID PolygonID, const int32 HoleNumber, const FVertexID EdgeVertexID0, const FVertexID EdgeVertexID1 ) const;
+	UFUNCTION( BlueprintPure, Category="Editable Mesh" ) FBox ComputeBoundingBox() const;
+	UFUNCTION( BlueprintPure, Category="Editable Mesh" ) FBoxSphereBounds ComputeBoundingBoxAndSphere() const;
+	UFUNCTION( BlueprintPure, Category="Editable Mesh" ) FVector ComputePolygonCenter( const FPolygonID PolygonID ) const;
+	UFUNCTION( BlueprintPure, Category="Editable Mesh" ) FPlane ComputePolygonPlane( const FPolygonID PolygonID ) const;
+	UFUNCTION( BlueprintPure, Category="Editable Mesh" ) FVector ComputePolygonNormal( const FPolygonID PolygonID ) const;
+	UFUNCTION( BlueprintPure, Category="Editable Mesh" ) const FSubdivisionLimitData& GetSubdivisionLimitData() const;
+	UFUNCTION( BlueprintPure, Category="Editable Mesh" ) void ComputePolygonTriangulation( const FPolygonID PolygonID, TArray<FMeshTriangle>& OutTriangles ) const;
+	UFUNCTION( BlueprintPure, Category="Editable Mesh" ) bool ComputeBarycentricWeightForPointOnPolygon( const FPolygonID PolygonID, const FVector PointOnPolygon, FMeshTriangle& OutTriangle, FVector& OutTriangleVertexWeights ) const;
+	UFUNCTION( BlueprintPure, Category="Editable Mesh" ) void ComputeTextureCoordinatesForPointOnPolygon( const FPolygonID PolygonID, const FVector PointOnPolygon, bool& bOutWereTextureCoordinatesFound, TArray<FVector4>& OutInterpolatedTextureCoordinates );
+	UFUNCTION( BlueprintPure, Category="Editable Mesh" ) void ComputePolygonsSharedEdges( const TArray<FPolygonID>& PolygonIDs, TArray<FEdgeID>& OutSharedEdgeIDs ) const;
+	UFUNCTION( BlueprintPure, Category="Editable Mesh" ) void FindPolygonLoop( const FEdgeID EdgeID, TArray<FEdgeID>& OutEdgeLoopEdgeIDs, TArray<FEdgeID>& OutFlippedEdgeIDs, TArray<FEdgeID>& OutReversedEdgeIDPathToTake, TArray<FPolygonID>& OutPolygonIDsToSplit ) const;
+
+
+	// Low poly editing features
+	// @todo mesheditor: Move elsewhere
+	UFUNCTION( BlueprintCallable, Category="Editable Mesh" ) void SetSubdivisionCount( const int32 NewSubdivisionCount );
+	UFUNCTION( BlueprintCallable, Category="Editable Mesh" ) void MoveVertices( const TArray<FVertexToMove>& VerticesToMove );
+	UFUNCTION( BlueprintCallable, Category="Editable Mesh" ) void CreateMissingPolygonPerimeterEdges( const FPolygonID PolygonID, TArray<FEdgeID>& OutNewEdgeIDs );
+	UFUNCTION( BlueprintCallable, Category="Editable Mesh" ) void CreateMissingPolygonHoleEdges( const FPolygonID PolygonID, const int32 HoleNumber, TArray<FEdgeID>& OutNewEdgeIDs );
+	UFUNCTION( BlueprintCallable, Category="Editable Mesh" ) void SplitEdge( const FEdgeID EdgeID, const TArray<float>& Splits, TArray<FVertexID>& OutNewVertexIDs );
+	UFUNCTION( BlueprintCallable, Category="Editable Mesh" ) void InsertEdgeLoop( const FEdgeID EdgeID, const TArray<float>& Splits, TArray<FEdgeID>& OutNewEdgeIDs );
+	UFUNCTION( BlueprintCallable, Category="Editable Mesh" ) void SplitPolygons( const TArray<FPolygonToSplit>& PolygonsToSplit, TArray<FEdgeID>& OutNewEdgeIDs );
+	UFUNCTION( BlueprintCallable, Category="Editable Mesh" ) void DeleteEdgeAndConnectedPolygons( const FEdgeID EdgeID, const bool bDeleteOrphanedEdges, const bool bDeleteOrphanedVertices, const bool bDeleteOrphanedVertexInstances, const bool bDeleteEmptyPolygonGroups );
+	UFUNCTION( BlueprintCallable, Category="Editable Mesh" ) void DeleteVertexAndConnectedEdgesAndPolygons( const FVertexID VertexID, const bool bDeleteOrphanedEdges, const bool bDeleteOrphanedVertices, const bool bDeleteOrphanedVertexInstances, const bool bDeleteEmptyPolygonGroups );
+	UFUNCTION( BlueprintCallable, Category="Editable Mesh" ) void DeleteOrphanVertices( const TArray<FVertexID>& VertexIDsToDelete );
+	UFUNCTION( BlueprintCallable, Category="Editable Mesh" ) void DeleteVertexInstances( const TArray<FVertexInstanceID>& VertexInstanceIDsToDelete, const bool bDeleteOrphanedVertices );
+	UFUNCTION( BlueprintCallable, Category="Editable Mesh" ) void DeleteEdges( const TArray<FEdgeID>& EdgeIDsToDelete, const bool bDeleteOrphanedVertices );
+	UFUNCTION( BlueprintCallable, Category="Editable Mesh" ) void CreateEmptyVertexRange( const int32 NumVerticesToCreate, TArray<FVertexID>& OutNewVertexIDs );
+	UFUNCTION( BlueprintCallable, Category="Editable Mesh" ) void CreateVertices( const TArray<FVertexToCreate>& VerticesToCreate, TArray<FVertexID>& OutNewVertexIDs );
+	UFUNCTION( BlueprintCallable, Category="Editable Mesh" ) void CreateVertexInstances( const TArray<FVertexInstanceToCreate>& VertexInstancesToCreate, TArray<FVertexInstanceID>& OutNewVertexInstanceIDs );
+	UFUNCTION( BlueprintCallable, Category="Editable Mesh" ) void CreateEdges( const TArray<FEdgeToCreate>& EdgesToCreate, TArray<FEdgeID>& OutNewEdgeIDs );
+	UFUNCTION( BlueprintCallable, Category="Editable Mesh" ) void CreatePolygons( const TArray<FPolygonToCreate>& PolygonsToCreate, TArray<FPolygonID>& OutNewPolygonIDs, TArray<FEdgeID>& OutNewEdgeIDs );
+	UFUNCTION( BlueprintCallable, Category="Editable Mesh" ) void DeletePolygons( const TArray<FPolygonID>& PolygonIDsToDelete, const bool bDeleteOrphanedEdges, const bool bDeleteOrphanedVertices, const bool bDeleteOrphanedVertexInstances, const bool bDeleteEmptyPolygonGroups );
+	UFUNCTION( BlueprintCallable, Category="Editable Mesh" ) void SetVerticesAttributes( const TArray<FAttributesForVertex>& AttributesForVertices );
+	UFUNCTION( BlueprintCallable, Category="Editable Mesh" ) void SetVertexInstancesAttributes( const TArray<FAttributesForVertexInstance>& AttributesForVertexInstances );
+	UFUNCTION( BlueprintCallable, Category="Editable Mesh" ) void SetEdgesAttributes( const TArray<FAttributesForEdge>& AttributesForEdges );
+	UFUNCTION( BlueprintCallable, Category="Editable Mesh" ) void SetPolygonsVertexAttributes( const TArray<FVertexAttributesForPolygon>& VertexAttributesForPolygons );
+	UFUNCTION( BlueprintCallable, Category="Editable Mesh" ) void ChangePolygonsVertexInstances( const TArray<FChangeVertexInstancesForPolygon>& VertexInstancesForPolygons );
+	UFUNCTION( BlueprintCallable, Category="Editable Mesh" ) void TryToRemovePolygonEdge( const FEdgeID EdgeID, bool& bOutWasEdgeRemoved, FPolygonID& OutNewPolygonID );
+	UFUNCTION( BlueprintCallable, Category="Editable Mesh" ) void TryToRemoveVertex( const FVertexID VertexID, bool& bOutWasVertexRemoved, FEdgeID& OutNewEdgeID );
+	UFUNCTION( BlueprintCallable, Category="Editable Mesh" ) void ExtrudePolygons( const TArray<FPolygonID>& Polygons, const float ExtrudeDistance, const bool bKeepNeighborsTogether, TArray<FPolygonID>& OutNewExtrudedFrontPolygons );
+	UFUNCTION( BlueprintCallable, Category="Editable Mesh" ) void ExtendEdges( const TArray<FEdgeID>& EdgeIDs, const bool bWeldNeighbors, TArray<FEdgeID>& OutNewExtendedEdgeIDs );
+	UFUNCTION( BlueprintCallable, Category="Editable Mesh" ) void ExtendVertices( const TArray<FVertexID>& VertexIDs, const bool bOnlyExtendClosestEdge, const FVector ReferencePosition, TArray<FVertexID>& OutNewExtendedVertexIDs );
+	UFUNCTION( BlueprintCallable, Category="Editable Mesh" ) void InsetPolygons( const TArray<FPolygonID>& PolygonIDs, const float InsetFixedDistance, const float InsetProgressTowardCenter, const EInsetPolygonsMode Mode, TArray<FPolygonID>& OutNewCenterPolygonIDs, TArray<FPolygonID>& OutNewSidePolygonIDs );
+	UFUNCTION( BlueprintCallable, Category="Editable Mesh" ) void BevelPolygons( const TArray<FPolygonID>& PolygonIDs, const float BevelFixedDistance, const float BevelProgressTowardCenter, TArray<FPolygonID>& OutNewCenterPolygonIDs, TArray<FPolygonID>& OutNewSidePolygonIDs );
+	UFUNCTION( BlueprintCallable, Category="Editable Mesh" ) void SetVerticesCornerSharpness( const TArray<FVertexID>& VertexIDs, const TArray<float>& VerticesNewCornerSharpness );
+	UFUNCTION( BlueprintCallable, Category="Editable Mesh" ) void SetEdgesCreaseSharpness( const TArray<FEdgeID>& EdgeIDs, const TArray<float>& EdgesNewCreaseSharpness );
+	UFUNCTION( BlueprintCallable, Category="Editable Mesh" ) void SetEdgesHardness( const TArray<FEdgeID>& EdgeIDs, const TArray<bool>& EdgesNewIsHard );
+	UFUNCTION( BlueprintCallable, Category="Editable Mesh" ) void SetEdgesHardnessAutomatically( const TArray<FEdgeID>& EdgeIDs, const float MaxDotProductForSoftEdge );	// @todo mesheditor: Not used for anything yet.  Remove it?  Use it during import/convert?
+	UFUNCTION( BlueprintCallable, Category="Editable Mesh" ) void SetEdgesVertices( const TArray<FVerticesForEdge>& VerticesForEdges );
+	UFUNCTION( BlueprintCallable, Category="Editable Mesh" ) void InsertPolygonPerimeterVertices( const FPolygonID PolygonID, const int32 InsertBeforeVertexNumber, const TArray<FVertexAndAttributes>& VerticesToInsert );
+	UFUNCTION( BlueprintCallable, Category="Editable Mesh" ) void RemovePolygonPerimeterVertices( const FPolygonID PolygonID, const int32 FirstVertexNumberToRemove, const int32 NumVerticesToRemove, const bool bDeleteOrphanedVertexInstances );
+	UFUNCTION( BlueprintCallable, Category="Editable Mesh" ) void FlipPolygons( const TArray<FPolygonID>& PolygonIDs );
+	UFUNCTION( BlueprintCallable, Category="Editable Mesh" ) void TriangulatePolygons( const TArray<FPolygonID>& PolygonIDs, TArray<FPolygonID>& OutNewTrianglePolygons );
+	UFUNCTION( BlueprintCallable, Category="Editable Mesh" ) void CreatePolygonGroups( const TArray<FPolygonGroupToCreate>& PolygonGroupsToCreate, TArray<FPolygonGroupID>& OutNewPolygonGroupIDs );
+	UFUNCTION( BlueprintCallable, Category="Editable Mesh" ) void DeletePolygonGroups( const TArray<FPolygonGroupID>& PolygonGroupIDs );
+	UFUNCTION( BlueprintCallable, Category="Editable Mesh" ) void AssignMaterialToPolygons( const TArray<FPolygonID>& PolygonIDs, UMaterialInterface* Material, TArray<FPolygonID>& OutNewPolygonIDs );
+	UFUNCTION( BlueprintCallable, Category="Editable Mesh" ) void WeldVertices( const TArray<FVertexID>& VertexIDs, FVertexID& OutNewVertexID );
+	UFUNCTION( BlueprintCallable, Category="Editable Mesh" ) void TessellatePolygons( const TArray<FPolygonID>& PolygonIDs, const ETriangleTessellationMode TriangleTessellationMode, TArray<FPolygonID>& OutNewPolygonIDs );
+	UFUNCTION( BlueprintCallable, Category="Editable Mesh" ) void SetTextureCoordinateCount( const int32 NumTexCoords );
+	UFUNCTION( BlueprintCallable, Category="Editable Mesh" ) void QuadrangulateMesh( TArray<FPolygonID>& OutNewPolygonIDs );
+	UFUNCTION( BlueprintCallable, Category="Editable Mesh" ) void GeneratePolygonTangentsAndNormals( const TArray<FPolygonID>& PolygonIDs );
+
+protected:
+
+	/** Adds a new change that can be used to undo a modification that happened to the mesh.  If bAllowUndo is false, then
+	    the undo will not be stored */
+	void AddUndo( TUniquePtr<FChange> NewUndo );
+
+public:
+	// @todo mesheditor: temporarily changed access to public so the adapter can call it when building the editable mesh from the source static mesh. Think about this some more.
+	/** Refreshes the entire OpenSubdiv state for this mesh and generates subdivision geometry (if the mesh is configured to have subdivision levels) */
+	void RefreshOpenSubdiv();
+
+protected:
+
+	/** Generates new limit surface data and caches it (if the mesh is configured to have subdivision levels).  This can be used to update the
+	    limit surface after a non-topology-effective edit has happened */
+	void GenerateOpenSubdivLimitSurfaceData();
+
+
+private:
+
+	void BevelOrInsetPolygons( const TArray<FPolygonID>& PolygonIDs, const float InsetFixedDistance, const float InsetProgressTowardCenter, const EInsetPolygonsMode Mode, const bool bShouldBevel, TArray<FPolygonID>& OutNewCenterPolygonIDs, TArray<FPolygonID>& OutNewSidePolygonIDs );
+
+	/** Method used internally by the quadrangulate action */
+	void QuadrangulatePolygonGroup( const FPolygonGroupID PolygonGroupID, TArray<FPolygonID>& OutNewPolygonIDs );
+
+	/** Called during end modification to generate tangents and normals on the pending polygon list */
+	void GenerateTangentsAndNormals();
+
+	/** Called during end modification to retriangulate polygons in the pending polygon list */
+	void RetriangulatePolygons();
+
+public:
+
+	/** Each editable vertex in this mesh. */
+	TSparseArray<FMeshVertex> Vertices;
+
+	/** Sparse array of rendering vertices, that matches the vertices in the mesh vertex buffers */
+	TSparseArray<FMeshVertexInstance> VertexInstances;
+
+	/** All editable mesh edges.  Note that some of these edges will be internal polygon edges, synthesized while
+	triangulating polygons into triangles.  Static meshes currently only support triangles. */
+	TSparseArray<FMeshEdge> Edges;
+
+	/** All of the polygons in this mesh */
+	TSparseArray<FMeshPolygon> Polygons;
+
+	/** All of the polygon groups in this mesh */
+	TSparseArray<FMeshPolygonGroup> PolygonGroups;
+
+// @todo mesheditor: sort out member access. Currently StaticMesh adapter relies on accessing this stuff directly
+//protected:
+
+	/** The sub-mesh we came from */
+	FEditableMeshSubMeshAddress SubMeshAddress;
+
+	/** True if undo features are enabled on this mesh.  You're only allowed to call MakeUndo() if this is set to true. */
+	bool bAllowUndo;
+
+	/** When bAllowUndo is enabled, this will store the changes that can be applied to revert anything that happened to this
+	    mesh since the last time that MakeUndo() was called. */
+	TUniquePtr<FCompoundChangeInput> Undo;
+
+	/** Adapters registered with this editable mesh */
+	UPROPERTY()
+	TArray<UEditableMeshAdapter*> Adapters;
+
+	/** The number of texture coordinates stored on the vertices of this mesh */
+	UPROPERTY( BlueprintReadOnly, Category="Editable Mesh" )
+	int32 TextureCoordinateCount;
+
+	/** How many levels to subdivide this mesh.  Zero will turn off subdivisions */
+	UPROPERTY( BlueprintReadOnly, Category="Editable Mesh" ) 
+	int32 SubdivisionCount;
+
+	/** List of polygons which need their tangent basis recalculating (and consequently their associated vertex instances) */
+	TSet<FPolygonID> PolygonsPendingNewTangentBasis;
+
+	/** List of polygons requiring retriangulation */
+	TSet<FPolygonID> PolygonsPendingTriangulation;
+
+	/** True if StartModification() has been called.  Call EndModification() when you've finished changing the mesh. */
+	bool bIsBeingModified;
+
+	/** While the mesh is being edited (between calls to StartModification() and EndModification()), this is the type of modification being performed */
+	EMeshModificationType CurrentModificationType;
+
+	/** While the mesh is being edited (between calls to StartModification() and EndModification()), stores whether topology could be affected */
+	EMeshTopologyChange CurrentToplogyChange;
+
+	/** Counter to determine when we should compact data */
+	UPROPERTY()
+	int32 PendingCompactCounter;
+
+	/** Data will be compacted after this many topology modifying actions. */
+	static const int32 CompactFrequency = 10;
+
+	/** OpenSubdiv topology refiner object.  This is generated for meshes that have subdivision levels, and reused to generate new limit surfaces 
+	    when geometry is moved.  When the mesh's topology changes, this object is regenerated from scratch. */
+	TSharedPtr<OpenSubdiv::v3_2_0::Far::TopologyRefiner> OsdTopologyRefiner;
+
+	/** Various cached arrays of mesh data in the form that OpenSubdiv expects to read it.  Required by GenerateOpenSubdivLimitSurfaceData(). */
+	TArray<int> OsdNumVerticesPerFace;
+	TArray<int> OsdVertexIndicesPerFace;
+	TArray<int> OsdCreaseVertexIndexPairs;
+	TArray<float> OsdCreaseWeights;
+	TArray<int> OsdCornerVertexIndices;
+	TArray<float> OsdCornerWeights;
+
+	TArray<int> OsdFVarIndicesPerFace;
+	struct FOsdFVarChannel
+	{
+		int ValueCount;
+		int const* ValueIndices;
+	};
+	TArray<FOsdFVarChannel> OsdFVarChannels;
+
+	/** The resulting limit surface geometry after GenerateOpenSubdivLimitSurfaceData() is called */
+	FSubdivisionLimitData SubdivisionLimitData;
+};
diff --git a/Engine/Plugins/Runtime/EditableMesh/Source/EditableMesh/Public/EditableMeshAdapter.h b/Engine/Plugins/Runtime/EditableMesh/Source/EditableMesh/Public/EditableMeshAdapter.h
new file mode 100644
--- /dev/null
+++ b/Engine/Plugins/Runtime/EditableMesh/Source/EditableMesh/Public/EditableMeshAdapter.h
@@ -0,0 +1,49 @@
+// Copyright 1998-2015 Epic Games, Inc. All Rights Reserved.
+
+#pragma once
+
+#include "CoreMinimal.h"
+#include "UObject/Object.h"
+#include "EditableMesh.h"
+#include "EditableMeshTypes.h"
+#include "EditableMeshAdapter.generated.h"
+
+
+UCLASS(Abstract)
+class UEditableMeshAdapter : public UObject
+{
+	GENERATED_BODY()
+
+public:
+
+	virtual void OnRebuildRenderMeshStart( const UEditableMesh* EditableMesh, const bool bRefreshBounds, const bool bInvalidateLighting ) PURE_VIRTUAL(,);
+	virtual void OnRebuildRenderMesh( const UEditableMesh* EditableMesh ) PURE_VIRTUAL(,);
+	virtual void OnRebuildRenderMeshFinish( const UEditableMesh* EditableMesh, const bool bUpdateCollision ) PURE_VIRTUAL(,);
+	virtual void OnStartModification( const UEditableMesh* EditableMesh, const EMeshModificationType MeshModificationType, const EMeshTopologyChange MeshTopologyChange ) PURE_VIRTUAL(,);
+	virtual void OnEndModification( const UEditableMesh* EditableMesh ) PURE_VIRTUAL(,);
+	virtual void OnReindexElements( const UEditableMesh* EditableMesh, const FElementIDRemappings& Remappings ) PURE_VIRTUAL(,);
+	virtual bool IsCommitted( const UEditableMesh* EditableMesh ) const PURE_VIRTUAL(,return false;);
+	virtual bool IsCommittedAsInstance( const UEditableMesh* EditableMesh ) const PURE_VIRTUAL(,return false;);
+	virtual void OnCommit( UEditableMesh* EditableMesh ) PURE_VIRTUAL(,);
+	virtual UEditableMesh* OnCommitInstance( UEditableMesh* EditableMesh, UPrimitiveComponent* ComponentToInstanceTo ) PURE_VIRTUAL(,return nullptr;);
+	virtual void OnRevert( UEditableMesh* EditableMesh ) PURE_VIRTUAL(,);
+	virtual UEditableMesh* OnRevertInstance( UEditableMesh* EditableMesh ) PURE_VIRTUAL(,return nullptr;);
+	virtual void OnPropagateInstanceChanges( UEditableMesh* EditableMesh ) PURE_VIRTUAL(,);
+
+	virtual void OnDeleteVertexInstances( const UEditableMesh* EditableMesh, const TArray<FVertexInstanceID>& VertexInstanceIDs ) PURE_VIRTUAL(,);
+	virtual void OnDeleteOrphanVertices( const UEditableMesh* EditableMesh, const TArray<FVertexID>& VertexIDs ) PURE_VIRTUAL(,);
+	virtual void OnCreateEmptyVertexRange( const UEditableMesh* EditableMesh, const TArray<FVertexID>& VertexIDs ) PURE_VIRTUAL(,);
+	virtual void OnCreateVertices( const UEditableMesh* EditableMesh, const TArray<FVertexID>& VertexIDs ) PURE_VIRTUAL(,);
+	virtual void OnCreateVertexInstances( const UEditableMesh* EditableMesh, const TArray<FVertexInstanceID>& VertexInstanceIDs ) PURE_VIRTUAL(,);
+	virtual void OnSetVertexAttribute( const UEditableMesh* EditableMesh, const FVertexID VertexID, const FName AttributeName, const int32 AttributeIndex, const FVector4 AttributeValue ) PURE_VIRTUAL(,);
+	virtual void OnSetVertexInstanceAttribute( const UEditableMesh* EditableMesh, const FVertexInstanceID VertexInstanceID, const FName AttributeName, const int32 AttributeIndex, const FVector4 AttributeValue ) PURE_VIRTUAL(,);
+	virtual void OnCreateEdges( const UEditableMesh* EditableMesh, const TArray<FEdgeID>& EdgeIDs ) PURE_VIRTUAL(,);
+	virtual void OnDeleteEdges( const UEditableMesh* EditableMesh, const TArray<FEdgeID>& EdgeIDs ) PURE_VIRTUAL(,);
+	virtual void OnSetEdgeAttribute( const UEditableMesh* EditableMesh, const FEdgeID EdgeID, const FName AttributeName, const int32 AttributeIndex, const FVector4 AttributeValue ) PURE_VIRTUAL(,);
+	virtual void OnCreatePolygons( const UEditableMesh* EditableMesh, const TArray<FPolygonID>& PolygonIDs ) PURE_VIRTUAL(,);
+	virtual void OnDeletePolygons( const UEditableMesh* EditableMesh, const TArray<FPolygonID>& PolygonIDs ) PURE_VIRTUAL(,);
+	virtual void OnChangePolygonVertexInstances( const UEditableMesh* EditableMesh, const TArray<FPolygonID>& PolygonIDs ) PURE_VIRTUAL(,);
+	virtual void OnCreatePolygonGroups( const UEditableMesh* EditableMesh, const TArray<FPolygonGroupID>& PolygonGroupIDs ) PURE_VIRTUAL(,);
+	virtual void OnDeletePolygonGroups( const UEditableMesh* EditableMesh, const TArray<FPolygonGroupID>& PolygonGroupIDs ) PURE_VIRTUAL(,);
+	virtual void OnRetriangulatePolygons( const UEditableMesh* EditableMesh, const TArray<FPolygonID>& PolygonIDs ) PURE_VIRTUAL(,);
+};
diff --git a/Engine/Plugins/Runtime/EditableMesh/Source/EditableMesh/Public/EditableMeshFactory.h b/Engine/Plugins/Runtime/EditableMesh/Source/EditableMesh/Public/EditableMeshFactory.h
new file mode 100644
--- /dev/null
+++ b/Engine/Plugins/Runtime/EditableMesh/Source/EditableMesh/Public/EditableMeshFactory.h
@@ -0,0 +1,28 @@
+// Copyright 1998-2015 Epic Games, Inc. All Rights Reserved.
+
+#pragma once
+
+#include "CoreMinimal.h"
+#include "EditableMeshTypes.h"
+#include "EditableMeshFactory.generated.h"
+
+
+// @todo mesheditor: Comment these classes and enums!
+
+
+UCLASS()
+class EDITABLEMESH_API UEditableMeshFactory : public UObject
+{
+	GENERATED_BODY()
+
+public:
+
+	static FEditableMeshSubMeshAddress MakeSubmeshAddress( class UPrimitiveComponent* PrimitiveComponent, const int32 LODIndex );
+	static FEditableMeshSubMeshAddress MakeSubmeshAddress( class UStaticMesh& StaticMesh, const int32 LODIndex );
+
+	static UEditableMesh* MakeEditableMesh( class UPrimitiveComponent* PrimitiveComponent, const FEditableMeshSubMeshAddress& SubMeshAddress );
+
+	UFUNCTION( BlueprintCallable, Category = "Editable Mesh" )
+	static UEditableMesh* MakeEditableMesh( class UPrimitiveComponent* PrimitiveComponent, const int32 LODIndex );
+
+};
diff --git a/Engine/Plugins/Runtime/EditableMesh/Source/EditableMesh/Public/IEditableMeshFormat.h b/Engine/Plugins/Runtime/EditableMesh/Source/EditableMesh/Public/IEditableMeshFormat.h
new file mode 100644
--- /dev/null
+++ b/Engine/Plugins/Runtime/EditableMesh/Source/EditableMesh/Public/IEditableMeshFormat.h
@@ -0,0 +1,22 @@
+// Copyright 1998-2015 Epic Games, Inc. All Rights Reserved.
+
+#pragma once
+
+#include "CoreMinimal.h"
+#include "Features/IModularFeature.h"
+#include "Features/IModularFeatures.h"
+
+
+/**
+ * Interface for mesh primitives that can be interactively edited (low-poly editing, sculpting, attribute painting, etc.)
+ */
+class IEditableMeshFormat : public IModularFeature
+{
+
+public:
+
+	// @todo mesheditor: Comments!
+	virtual void FillMeshObjectPtr( class UPrimitiveComponent& Component, struct FEditableMeshSubMeshAddress& SubMeshAddress ) = 0;
+	virtual class UEditableMesh* MakeEditableMesh( class UPrimitiveComponent& Component, const struct FEditableMeshSubMeshAddress& SubMeshAddress ) = 0;
+
+};
\ No newline at end of file
diff --git a/Engine/Plugins/Runtime/EditableMesh/Source/EditableMesh/Public/IEditableMeshModule.h b/Engine/Plugins/Runtime/EditableMesh/Source/EditableMesh/Public/IEditableMeshModule.h
new file mode 100644
--- /dev/null
+++ b/Engine/Plugins/Runtime/EditableMesh/Source/EditableMesh/Public/IEditableMeshModule.h
@@ -0,0 +1,39 @@
+// Copyright 1998-2015 Epic Games, Inc. All Rights Reserved.
+
+#pragma once
+
+#include "CoreMinimal.h"
+#include "ModuleManager.h"
+
+
+/**
+ * The public interface to this module
+ */
+class IEditableMeshModule : public IModuleInterface
+{
+
+public:
+
+	/**
+	 * Singleton-like access to this module's interface.  This is just for convenience!
+	 * Beware of calling this during the shutdown phase, though.  Your module might have been unloaded already.
+	 *
+	 * @return Returns singleton instance, loading the module on demand if needed
+	 */
+	static inline IEditableMeshModule& Get()
+	{
+		return FModuleManager::LoadModuleChecked< IEditableMeshModule >( "EditableMesh" );
+	}
+
+	/**
+	 * Checks to see if this module is loaded and ready.  It is only valid to call Get() if IsAvailable() returns true.
+	 *
+	 * @return True if the module is loaded and ready to use
+	 */
+	static inline bool IsAvailable()
+	{
+		return FModuleManager::Get().IsModuleLoaded( "EditableMesh" );
+	}
+
+};
+
diff --git a/Engine/Plugins/Runtime/EditableMesh/Source/EditableMesh/StaticMeshEditableMeshFormat.cpp b/Engine/Plugins/Runtime/EditableMesh/Source/EditableMesh/StaticMeshEditableMeshFormat.cpp
new file mode 100644
--- /dev/null
+++ b/Engine/Plugins/Runtime/EditableMesh/Source/EditableMesh/StaticMeshEditableMeshFormat.cpp
@@ -0,0 +1,64 @@
+// Copyright 1998-2015 Epic Games, Inc. All Rights Reserved.
+
+#include "StaticMeshEditableMeshFormat.h"
+#include "EditableMesh.h"
+#include "Engine/StaticMesh.h"
+#include "Components/StaticMeshComponent.h"
+#include "StaticMeshResources.h"
+#include "EditableStaticMeshAdapter.h"
+
+
+void FStaticMeshEditableMeshFormat::FillMeshObjectPtr( UPrimitiveComponent& Component, FEditableMeshSubMeshAddress& SubMeshAddress )
+{
+	SubMeshAddress.MeshObjectPtr = nullptr;
+
+	const UStaticMeshComponent* StaticMeshComponentPtr = Cast<const UStaticMeshComponent>( &Component );
+	if( StaticMeshComponentPtr != nullptr )
+	{
+		const UStaticMeshComponent& StaticMeshComponent = *StaticMeshComponentPtr;
+
+		UStaticMesh* ComponentStaticMesh = StaticMeshComponent.GetStaticMesh();
+		if( ComponentStaticMesh != nullptr && ComponentStaticMesh->HasValidRenderData() )
+		{
+			SubMeshAddress.MeshObjectPtr = ComponentStaticMesh;
+		}
+	}
+}
+
+
+UEditableMesh* FStaticMeshEditableMeshFormat::MakeEditableMesh( UPrimitiveComponent& Component, const FEditableMeshSubMeshAddress& SubMeshAddress )
+{
+	// If the static mesh already has an attached UEditableStaticMeshAdapter, use that in preference to creating a new one
+	const UStaticMeshComponent* StaticMeshComponentPtr = Cast<const UStaticMeshComponent>( &Component );
+	if( StaticMeshComponentPtr != nullptr )
+	{
+		const UStaticMeshComponent& StaticMeshComponent = *StaticMeshComponentPtr;
+		if( StaticMeshComponent.GetStaticMesh() != nullptr )
+		{
+			UEditableMesh* EditableMesh = Cast<UEditableMesh>( StaticMeshComponent.GetStaticMesh()->EditableMesh );
+			if( EditableMesh )
+			{
+				EditableMesh->SetSubMeshAddress( SubMeshAddress );
+				return EditableMesh;
+			}
+		}
+	}
+
+	UEditableMesh* EditableMesh = NewObject<UEditableMesh>();
+	UEditableStaticMeshAdapter* EditableStaticMesh = NewObject<UEditableStaticMeshAdapter>( EditableMesh );
+	EditableMesh->Adapters.Add( EditableStaticMesh );
+
+	EditableStaticMesh->InitEditableStaticMesh( EditableMesh, Component, SubMeshAddress );
+
+	// Don't bother returning a new mesh if it has no geometry
+	if( EditableMesh->GetVertexCount() == 0 )
+	{
+		EditableMesh->Adapters.Remove( EditableStaticMesh );
+		EditableStaticMesh->MarkPendingKill();
+		EditableMesh->MarkPendingKill();
+		EditableMesh = nullptr;
+	}
+
+	return EditableMesh;
+}
+
diff --git a/Engine/Plugins/Runtime/EditableMesh/Source/EditableMesh/StaticMeshEditableMeshFormat.h b/Engine/Plugins/Runtime/EditableMesh/Source/EditableMesh/StaticMeshEditableMeshFormat.h
new file mode 100644
--- /dev/null
+++ b/Engine/Plugins/Runtime/EditableMesh/Source/EditableMesh/StaticMeshEditableMeshFormat.h
@@ -0,0 +1,21 @@
+// Copyright 1998-2015 Epic Games, Inc. All Rights Reserved.
+
+#pragma once
+
+#include "CoreMinimal.h"
+#include "IEditableMeshFormat.h"
+
+
+/**
+ * Implements interactive mesh editing support for Static Meshes
+ */ 
+class FStaticMeshEditableMeshFormat : public IEditableMeshFormat
+{
+
+public:
+
+	// IEditableMeshFormat interface
+	virtual void FillMeshObjectPtr( class UPrimitiveComponent& Component, FEditableMeshSubMeshAddress& SubMeshAddress ) override;
+	virtual UEditableMesh* MakeEditableMesh( class UPrimitiveComponent& Component, const FEditableMeshSubMeshAddress& SubMeshAddress ) override;
+
+};
\ No newline at end of file
diff --git a/Engine/Source/Runtime/Engine/Engine.Build.cs b/Engine/Source/Runtime/Engine/Engine.Build.cs
--- a/Engine/Source/Runtime/Engine/Engine.Build.cs
+++ b/Engine/Source/Runtime/Engine/Engine.Build.cs
@@ -173,7 +173,6 @@ public class Engine : ModuleRules
 				"MovieSceneTracks",
 				"HeadMountedDisplay",
                 "MRMesh",
-				"MeshEditingRuntime",
                 "StreamingPauseRendering",
                 "Niagara",
 			}
diff --git a/Engine/Source/Runtime/Engine/Private/UnrealEngine.cpp b/Engine/Source/Runtime/Engine/Private/UnrealEngine.cpp
--- a/Engine/Source/Runtime/Engine/Private/UnrealEngine.cpp
+++ b/Engine/Source/Runtime/Engine/Private/UnrealEngine.cpp
@@ -1055,7 +1055,6 @@ void UEngine::Init(IEngineLoop* InEngineLoop)
 	// Dynamically load engine runtime modules
 	{
 		FModuleManager::Get().LoadModuleChecked(TEXT("StreamingPauseRendering"));
-		FModuleManager::Get().LoadModuleChecked(TEXT("MeshEditingRuntime"));
 		FModuleManager::Get().LoadModuleChecked(TEXT("Niagara"));
 		FModuleManager::Get().LoadModuleChecked(TEXT("GeometryCache"));
 		FModuleManager::Get().LoadModuleChecked(TEXT("MovieScene"));
