Mesh Editor: New "EditableMesh" plugin (replaces "MeshEditingRuntime" module) (part 2)
- Added missing file

diff --git a/Engine/Plugins/Runtime/EditableMesh/EditableMesh.uplugin b/Engine/Plugins/Runtime/EditableMesh/EditableMesh.uplugin
new file mode 100644
--- /dev/null
+++ b/Engine/Plugins/Runtime/EditableMesh/EditableMesh.uplugin
@@ -0,0 +1,25 @@
+{
+	"FileVersion" : 3,
+	"Version" : 1,
+	"VersionName" : "1.0",
+	"FriendlyName" : "Editable Mesh",
+	"Description" : "A mesh with rich runtime editing features",
+	"Category" : "Rendering",
+	"CreatedBy" : "Epic Games, Inc.",
+	"CreatedByURL" : "http://epicgames.com",
+	"DocsURL" : "",
+	"MarketplaceURL" : "",
+	"SupportURL" : "",
+	"EnabledByDefault" : true,
+	"CanContainContent" : false,
+	"IsBetaVersion" : false,
+	"Installed" : false,
+	"Modules" :
+	[
+		{
+			"Name" : "EditableMesh",
+			"Type" : "Runtime",
+			"LoadingPhase" : "Default"
+		}
+	]
+}
\ No newline at end of file
