Put compact frequency back up to every 10 topology changing transactions.

diff --git a/Engine/Plugins/Runtime/EditableMesh/Source/EditableMesh/Public/EditableMesh.h b/Engine/Plugins/Runtime/EditableMesh/Source/EditableMesh/Public/EditableMesh.h
--- a/Engine/Plugins/Runtime/EditableMesh/Source/EditableMesh/Public/EditableMesh.h
+++ b/Engine/Plugins/Runtime/EditableMesh/Source/EditableMesh/Public/EditableMesh.h
@@ -880,7 +880,7 @@ public:
 	int32 PendingCompactCounter;
 
 	/** Data will be compacted after this many topology modifying actions. */
-	static const int32 CompactFrequency = 1;
+	static const int32 CompactFrequency = 10;
 
 	/** OpenSubdiv topology refiner object.  This is generated for meshes that have subdivision levels, and reused to generate new limit surfaces 
 	    when geometry is moved.  When the mesh's topology changes, this object is regenerated from scratch. */
