Mesh Editor: Fixed cases where octree was not updated after polygons changed
- Also, renamed a function

diff --git a/Engine/Plugins/Runtime/EditableMesh/Source/EditableMesh/EditableMesh.cpp b/Engine/Plugins/Runtime/EditableMesh/Source/EditableMesh/EditableMesh.cpp
--- a/Engine/Plugins/Runtime/EditableMesh/Source/EditableMesh/EditableMesh.cpp
+++ b/Engine/Plugins/Runtime/EditableMesh/Source/EditableMesh/EditableMesh.cpp
@@ -516,7 +516,7 @@ void UEditableMesh::EndModification( const bool bFromUndo )
 		// Rebuild our octree
 		if( CurrentModificationType == EMeshModificationType::Final )
 		{
-			UpdateOctreeIncrementallyIfPossibleOtherwiseRebuildIt();
+			UpdateOrRebuildOctree();
 		}
 
 		FStartOrEndModificationChangeInput RevertInput;
@@ -7496,6 +7496,13 @@ void UEditableMesh::InsertPolygonPerimeterVertices( const FPolygonID PolygonID,
 	PolygonsPendingTriangulation.Add( PolygonID );
 	PolygonsPendingNewTangentBasis.Add( PolygonID );	// @todo mesheditor: and other vertex connected polygons too?
 
+	// Update spatial database
+	if( Octree.IsValid() )
+	{
+		DeletedOctreePolygonIDs.Add( PolygonID );
+		NewOctreePolygonIDs.Add( PolygonID );
+	}
+
 	UE_LOG( LogEditableMesh, Verbose, TEXT( "* InsertPolygonPerimeterVertices finished" ) );
 }
 
@@ -7555,6 +7562,13 @@ void UEditableMesh::RemovePolygonPerimeterVertices( const FPolygonID PolygonID,
 	PolygonsPendingTriangulation.Add( PolygonID );
 	PolygonsPendingNewTangentBasis.Add( PolygonID );	// @todo mesheditor: and other vertex connected polygons too?
 
+	// Update spatial database
+	if( Octree.IsValid() )
+	{
+		DeletedOctreePolygonIDs.Add( PolygonID );
+		NewOctreePolygonIDs.Add( PolygonID );
+	}
+
 	UE_LOG( LogEditableMesh, Verbose, TEXT( "* RemovePolygonPerimeterVertices finished" ) );
 }
 
@@ -9116,7 +9130,7 @@ void UEditableMesh::QuadrangulatePolygonGroup( const FPolygonGroupID PolygonGrou
 }
 
 
-void UEditableMesh::UpdateOctreeIncrementallyIfPossibleOtherwiseRebuildIt()
+void UEditableMesh::UpdateOrRebuildOctree()
 {
 	bool bAnythingChanged = true;
 	bool bDoIncrementalUpdate = false;
diff --git a/Engine/Plugins/Runtime/EditableMesh/Source/EditableMesh/Public/EditableMesh.h b/Engine/Plugins/Runtime/EditableMesh/Source/EditableMesh/Public/EditableMesh.h
--- a/Engine/Plugins/Runtime/EditableMesh/Source/EditableMesh/Public/EditableMesh.h
+++ b/Engine/Plugins/Runtime/EditableMesh/Source/EditableMesh/Public/EditableMesh.h
@@ -834,7 +834,7 @@ private:
 
 	/** Tries to incrementally update the octree based off geometry that has changed since last time.  If that's not
 	    reasonable to do, then this will rebuild the entire octree from scratch */
-	void UpdateOctreeIncrementallyIfPossibleOtherwiseRebuildIt();
+	void UpdateOrRebuildOctree();
 
 public:
 	// @todo mesheditor: temporarily changed access to public so the adapter can call it when building the editable mesh from the source static mesh. Think about this some more.
