Fixed subdivision mesh not being refreshed after normals and tangents were generated for a new mesh

diff --git a/Engine/Plugins/Runtime/EditableMesh/Source/EditableMesh/EditableStaticMeshAdapter.cpp b/Engine/Plugins/Runtime/EditableMesh/Source/EditableMesh/EditableStaticMeshAdapter.cpp
--- a/Engine/Plugins/Runtime/EditableMesh/Source/EditableMesh/EditableStaticMeshAdapter.cpp
+++ b/Engine/Plugins/Runtime/EditableMesh/Source/EditableMesh/EditableStaticMeshAdapter.cpp
@@ -462,8 +462,6 @@ void UEditableStaticMeshAdapter::InitEditableStaticMesh( UEditableMesh* Editable
 		}
 	}
 
-	EditableMesh->RefreshOpenSubdiv();
-
 	// Cache polygon tangent bases
 	static TArray<FPolygonID> PolygonIDs;
 	PolygonIDs.Reset();
@@ -490,6 +488,9 @@ void UEditableStaticMeshAdapter::InitEditableStaticMesh( UEditableMesh* Editable
 
 	EditableMesh->GenerateTangentsAndNormals();
 #endif
+
+	EditableMesh->RefreshOpenSubdiv();
+	EditableMesh->RebuildOctree();
 }
 
 
@@ -968,6 +969,7 @@ void UEditableStaticMeshAdapter::UpdateBoundsAndCollision( const UEditableMesh*
 
 	// Compute a new bounding box
 	// @todo mesheditor perf: Only do this if the bounds may have changed (need hinting)
+	// @todo mesheditor perf: This is hit every frame while dragging, and doesn't scale well to large meshes.  It's also computed for "preview" changes that will never be interacted with, rather than just on roll-back.  Can we keep bounds up to date as we go (grow while editing, then shrink perfectly after)?  Might cause gizmo/bounds pop.
 	{
 		FStaticMeshRenderData& StaticMeshRenderData = *StaticMesh->RenderData;
 
