Fixed a crash updating octree, caused by a double-delete

diff --git a/Engine/Plugins/Runtime/EditableMesh/Source/EditableMesh/EditableMesh.cpp b/Engine/Plugins/Runtime/EditableMesh/Source/EditableMesh/EditableMesh.cpp
--- a/Engine/Plugins/Runtime/EditableMesh/Source/EditableMesh/EditableMesh.cpp
+++ b/Engine/Plugins/Runtime/EditableMesh/Source/EditableMesh/EditableMesh.cpp
@@ -4766,15 +4766,18 @@ void UEditableMesh::DeletePolygons( const TArray<FPolygonID>& PolygonIDsToDelete
 
 			// Delete the polygon
 			Polygons.RemoveAt( PolygonID.GetValue() );
-		}
 
-		// Update spatial database
-		if( Octree.IsValid() )
-		{
-			DeletedOctreePolygonIDs.Append( PolygonIDsToDelete );
-			for( const FPolygonID PolygonID : PolygonIDsToDelete )
+			// Update spatial database
+			if( Octree.IsValid() )
 			{
-				NewOctreePolygonIDs.Remove( PolygonID );
+				// If the polygon we're deleting is still pending being added to the octree, the only thing we need to do is remove it from 
+				// the set of newly added polygons
+				if( NewOctreePolygonIDs.Remove( PolygonID ) == 0 )
+				{
+					// The polygon wasn't in our set of newly added polygons, so it must have been committed to the octree already.  We'll
+					// enqueue it for deletion here.
+					DeletedOctreePolygonIDs.Add( PolygonID );
+				}
 			}
 		}
 
