Avoid fully recomputing mesh bounds during interactive drags
- Static meshes now only inflate their bounding box while performing a drag.  The proper "tight" bounds will be later.
- This means the bounds is quite exaggerated until the interaction has completed.
- This greatly improves performance when editing large meshes

diff --git a/Engine/Plugins/Runtime/EditableMesh/Source/EditableMesh/EditableStaticMeshAdapter.cpp b/Engine/Plugins/Runtime/EditableMesh/Source/EditableMesh/EditableStaticMeshAdapter.cpp
--- a/Engine/Plugins/Runtime/EditableMesh/Source/EditableMesh/EditableStaticMeshAdapter.cpp
+++ b/Engine/Plugins/Runtime/EditableMesh/Source/EditableMesh/EditableStaticMeshAdapter.cpp
@@ -16,8 +16,9 @@ const FTriangleID FTriangleID::Invalid( TNumericLimits<uint32>::Max() );
 
 UEditableStaticMeshAdapter::UEditableStaticMeshAdapter()
 	: StaticMesh( nullptr ),
+	  StaticMeshLODIndex( 0 ),
 	  RecreateRenderStateContext(),
-	  StaticMeshLODIndex( 0 )
+	  CachedBoundingBoxAndSphere( FVector::ZeroVector, FVector::ZeroVector, 0.0f )
 {
 }
 
@@ -819,7 +820,13 @@ void UEditableStaticMeshAdapter::OnEndModification( const UEditableMesh* Editabl
 
 void UEditableStaticMeshAdapter::OnRebuildRenderMeshFinish( const UEditableMesh* EditableMesh, const bool bUpdateCollision )
 {
-	UpdateBoundsAndCollision( EditableMesh, bUpdateCollision );
+	const bool bShouldRecomputeBounds = bUpdateCollision;	// Only rebuild the bounds from scratch if we've finished a drag (bUpdateCollision will be true)
+	UpdateBounds( EditableMesh, bShouldRecomputeBounds );
+	
+	if( bUpdateCollision )
+	{
+		UpdateCollision();
+	}
 
 	StaticMesh->InitResources();
 
@@ -961,18 +968,12 @@ void UEditableStaticMeshAdapter::OnPropagateInstanceChanges( UEditableMesh* Edit
 }
 
 
-void UEditableStaticMeshAdapter::UpdateBoundsAndCollision( const UEditableMesh* EditableMesh, const bool bUpdateCollision )
+void UEditableStaticMeshAdapter::UpdateBounds( const UEditableMesh* EditableMesh, const bool bShouldRecompute )
 {
-	// @todo mesheditor: we will need to create a new DDC key once we are able to edit placed instances individually.
-	// Will need to find a way of deriving the key based on the mesh key and an instance number which remains constant,
-	// otherwise we risk filling the DDC with junk (i.e. using vertex positions etc is not scalable).
-
-	// Compute a new bounding box
-	// @todo mesheditor perf: Only do this if the bounds may have changed (need hinting)
-	// @todo mesheditor perf: This is hit every frame while dragging, and doesn't scale well to large meshes.  It's also computed for "preview" changes that will never be interacted with, rather than just on roll-back.  Can we keep bounds up to date as we go (grow while editing, then shrink perfectly after)?  Might cause gizmo/bounds pop.
+	if( bShouldRecompute )
 	{
-		FStaticMeshRenderData& StaticMeshRenderData = *StaticMesh->RenderData;
-
+		// Compute a new bounding box
+		// @todo mesheditor perf: During the final modification, only do this if the bounds may have changed (need hinting)
 		FBoxSphereBounds BoundingBoxAndSphere;
 
 		// @todo mesheditor LODs: Really we should store the bounds of LOD0 inside the static mesh.  Our editable mesh might be for a different LOD.
@@ -1021,67 +1022,68 @@ void UEditableStaticMeshAdapter::UpdateBoundsAndCollision( const UEditableMesh*
 				}
 			}
 		}
-	
-		StaticMeshRenderData.Bounds = BoundingBoxAndSphere;
-		StaticMesh->CalculateExtendedBounds();
-	}
 
+		this->CachedBoundingBoxAndSphere = BoundingBoxAndSphere;
+	}
+	
+	FStaticMeshRenderData& StaticMeshRenderData = *StaticMesh->RenderData;
+	StaticMeshRenderData.Bounds = CachedBoundingBoxAndSphere;
+	StaticMesh->CalculateExtendedBounds();
+}
 
-	// Refresh collision (only if the interaction has finished though -- this is really expensive!)
-	if( bUpdateCollision )
-	{
 
-		// @todo mesheditor collision: We're wiping the existing simplified collision and generating a simple bounding
-		// box collision, since that's the best we can do without impacting performance.  We always using visibility (complex)
-		// collision for traces while mesh editing (for hover/selection), so simplified collision isn't really important.
-		const bool bRecreateSimplifiedCollision = true;
+void UEditableStaticMeshAdapter::UpdateCollision()
+{
+	// @todo mesheditor collision: We're wiping the existing simplified collision and generating a simple bounding
+	// box collision, since that's the best we can do without impacting performance.  We always using visibility (complex)
+	// collision for traces while mesh editing (for hover/selection), so simplified collision isn't really important.
+	const bool bRecreateSimplifiedCollision = true;
 
-		if( StaticMesh->BodySetup == nullptr )
-		{
-			StaticMesh->CreateBodySetup();
-		}
+	if( StaticMesh->BodySetup == nullptr )
+	{
+		StaticMesh->CreateBodySetup();
+	}
 
-		UBodySetup* BodySetup = StaticMesh->BodySetup;
+	UBodySetup* BodySetup = StaticMesh->BodySetup;
 
-		// NOTE: We don't bother calling Modify() on the BodySetup as EndModification() will rebuild this guy after every undo
-		// BodySetup->Modify();
+	// NOTE: We don't bother calling Modify() on the BodySetup as EndModification() will rebuild this guy after every undo
+	// BodySetup->Modify();
 
-		if( bRecreateSimplifiedCollision )
+	if( bRecreateSimplifiedCollision )
+	{
+		if( BodySetup->AggGeom.GetElementCount() > 0 )
 		{
-			if( BodySetup->AggGeom.GetElementCount() > 0 )
-			{
-				BodySetup->RemoveSimpleCollision();
-			}
+			BodySetup->RemoveSimpleCollision();
 		}
+	}
 
-		BodySetup->InvalidatePhysicsData();
+	BodySetup->InvalidatePhysicsData();
 
-		if( bRecreateSimplifiedCollision )
-		{
-			const FBoxSphereBounds Bounds = StaticMesh->GetBounds();
-
-			FKBoxElem BoxElem;
-			BoxElem.Center = Bounds.Origin;
-			BoxElem.X = Bounds.BoxExtent.X * 2.0f;
-			BoxElem.Y = Bounds.BoxExtent.Y * 2.0f;
-			BoxElem.Z = Bounds.BoxExtent.Z * 2.0f;
-			BodySetup->AggGeom.BoxElems.Add( BoxElem );
-		}
+	if( bRecreateSimplifiedCollision )
+	{
+		const FBoxSphereBounds Bounds = StaticMesh->GetBounds();
+
+		FKBoxElem BoxElem;
+		BoxElem.Center = Bounds.Origin;
+		BoxElem.X = Bounds.BoxExtent.X * 2.0f;
+		BoxElem.Y = Bounds.BoxExtent.Y * 2.0f;
+		BoxElem.Z = Bounds.BoxExtent.Z * 2.0f;
+		BodySetup->AggGeom.BoxElems.Add( BoxElem );
+	}
 
-		// Update all static mesh components that are using this mesh
-		// @todo mesheditor perf: This is a pretty heavy operation, and overlaps with what we're already doing in RecreateRenderStateContext
-		// a little bit.  Ideally we do everything in a single pass.  Furthermore, if this could be updated lazily it would be faster.
+	// Update all static mesh components that are using this mesh
+	// @todo mesheditor perf: This is a pretty heavy operation, and overlaps with what we're already doing in RecreateRenderStateContext
+	// a little bit.  Ideally we do everything in a single pass.  Furthermore, if this could be updated lazily it would be faster.
+	{
+		for( FObjectIterator Iter( UStaticMeshComponent::StaticClass() ); Iter; ++Iter )
 		{
-			for( FObjectIterator Iter( UStaticMeshComponent::StaticClass() ); Iter; ++Iter )
+			UStaticMeshComponent* StaticMeshComponent = Cast<UStaticMeshComponent>( *Iter );
+			if( StaticMeshComponent->GetStaticMesh() == StaticMesh )
 			{
-				UStaticMeshComponent* StaticMeshComponent = Cast<UStaticMeshComponent>( *Iter );
-				if( StaticMeshComponent->GetStaticMesh() == StaticMesh )
+				// it needs to recreate IF it already has been created
+				if( StaticMeshComponent->IsPhysicsStateCreated() )
 				{
-					// it needs to recreate IF it already has been created
-					if( StaticMeshComponent->IsPhysicsStateCreated() )
-					{
-						StaticMeshComponent->RecreatePhysicsState();
-					}
+					StaticMeshComponent->RecreatePhysicsState();
 				}
 			}
 		}
@@ -1096,6 +1098,8 @@ void UEditableStaticMeshAdapter::OnSetVertexAttribute( const UEditableMesh* Edit
 
 	if( AttributeName == UEditableMeshAttribute::VertexPosition() )
 	{
+		const FVector NewVertexPosition = FVector( AttributeValue );
+
 		// @todo mesheditor: eventually break out subdivided mesh into a different adapter which handles things differently?
 		// (may also want different component eventually)
 		if( !EditableMesh->IsPreviewingSubdivisions() )
@@ -1104,9 +1108,30 @@ void UEditableStaticMeshAdapter::OnSetVertexAttribute( const UEditableMesh* Edit
 			for( const FVertexInstanceID VertexInstanceID : Vertex.VertexInstanceIDs )
 			{
 				check( EditableMesh->VertexInstances.IsAllocated( VertexInstanceID.GetValue() ) );
-				StaticMeshLOD.PositionVertexBuffer.VertexPosition( VertexInstanceID.GetValue() ) = FVector( AttributeValue );
+				StaticMeshLOD.PositionVertexBuffer.VertexPosition( VertexInstanceID.GetValue() ) = NewVertexPosition;
 			}
 		}
+
+		// Update cached bounds.  This is not a "perfect" bounding sphere and centered box.  Instead, we take our current bounds
+		// and inflate it to include the updated vertex position, translating the bounds proportionally to reduce how much it
+		// needs to be expanded.  The "perfect" bounds will be computed in UpdateBounds() when an interaction is finalized.
+		{
+			const FVector OffsetFromCenter = NewVertexPosition - CachedBoundingBoxAndSphere.Origin;
+			const float SquaredDistanceToCenter = OffsetFromCenter.SizeSquared();
+			const float SquaredSphereRadius = CachedBoundingBoxAndSphere.SphereRadius * CachedBoundingBoxAndSphere.SphereRadius;
+			if( SquaredDistanceToCenter > SquaredSphereRadius )
+			{
+				const float DistanceToCenter = FMath::Sqrt( SquaredDistanceToCenter );
+				const float RadiusDelta = ( DistanceToCenter - CachedBoundingBoxAndSphere.SphereRadius ) * 0.5f;
+				CachedBoundingBoxAndSphere.SphereRadius += RadiusDelta;
+				CachedBoundingBoxAndSphere.Origin += OffsetFromCenter * ( RadiusDelta / DistanceToCenter );
+			}
+
+			// Update extent
+			CachedBoundingBoxAndSphere.BoxExtent.X = FMath::Max( CachedBoundingBoxAndSphere.BoxExtent.X, FMath::Abs( NewVertexPosition.X - CachedBoundingBoxAndSphere.Origin.X ) );
+			CachedBoundingBoxAndSphere.BoxExtent.Y = FMath::Max( CachedBoundingBoxAndSphere.BoxExtent.Y, FMath::Abs( NewVertexPosition.Y - CachedBoundingBoxAndSphere.Origin.Y ) );
+			CachedBoundingBoxAndSphere.BoxExtent.Z = FMath::Max( CachedBoundingBoxAndSphere.BoxExtent.Z, FMath::Abs( NewVertexPosition.X - CachedBoundingBoxAndSphere.Origin.Z ) );
+		}
 	}
 }
 
diff --git a/Engine/Plugins/Runtime/EditableMesh/Source/EditableMesh/EditableStaticMeshAdapter.h b/Engine/Plugins/Runtime/EditableMesh/Source/EditableMesh/EditableStaticMeshAdapter.h
--- a/Engine/Plugins/Runtime/EditableMesh/Source/EditableMesh/EditableStaticMeshAdapter.h
+++ b/Engine/Plugins/Runtime/EditableMesh/Source/EditableMesh/EditableStaticMeshAdapter.h
@@ -169,7 +169,10 @@ private:
 	void UpdateIndexBufferFormatIfNeeded( const TArray<FMeshTriangle>& Triangles );
 
 	/** Rebuilds bounds */
-	void UpdateBoundsAndCollision( const UEditableMesh* EditableMesh, const bool bUpdateCollision );
+	void UpdateBounds( const UEditableMesh* EditableMesh, const bool bShouldRecompute );
+
+	/** Rebuilds collision.  Bounds should always be updated first. */
+	void UpdateCollision();
 
 	/** Gets the editable mesh section index which corresponds to the given rendering section index */
 	FPolygonGroupID GetSectionForRenderingSectionIndex( const int32 RenderingSectionIndex ) const;
@@ -196,7 +199,10 @@ private:
 	/** All of the polygon groups in this mesh */
 	TSparseArray< FRenderingPolygonGroup > RenderingPolygonGroups;
 
-
 	/** Used to refresh all components in the scene that may be using a mesh we're editing */
 	TSharedPtr< class FStaticMeshComponentRecreateRenderStateContext > RecreateRenderStateContext;
+
+	/** Cached bounding box for the static mesh.  This bounds can be (temporarily) larger than the actual mesh itself as
+	    an optimization. */
+	FBoxSphereBounds CachedBoundingBoxAndSphere;
 };
