Avoid fully recomputing mesh bounds during interactive drags (part 2)
- Fixed the initial cached bounding box not matching the original static mesh

diff --git a/Engine/Plugins/Runtime/EditableMesh/Source/EditableMesh/EditableStaticMeshAdapter.cpp b/Engine/Plugins/Runtime/EditableMesh/Source/EditableMesh/EditableStaticMeshAdapter.cpp
--- a/Engine/Plugins/Runtime/EditableMesh/Source/EditableMesh/EditableStaticMeshAdapter.cpp
+++ b/Engine/Plugins/Runtime/EditableMesh/Source/EditableMesh/EditableStaticMeshAdapter.cpp
@@ -463,6 +463,8 @@ void UEditableStaticMeshAdapter::InitEditableStaticMesh( UEditableMesh* Editable
 		}
 	}
 
+	this->CachedBoundingBoxAndSphere = StaticMesh->GetBounds();
+
 	// Cache polygon tangent bases
 	static TArray<FPolygonID> PolygonIDs;
 	PolygonIDs.Reset();
