Fixed another crash updating octree, caused by a double-delete

diff --git a/Engine/Plugins/Runtime/EditableMesh/Source/EditableMesh/EditableMesh.cpp b/Engine/Plugins/Runtime/EditableMesh/Source/EditableMesh/EditableMesh.cpp
--- a/Engine/Plugins/Runtime/EditableMesh/Source/EditableMesh/EditableMesh.cpp
+++ b/Engine/Plugins/Runtime/EditableMesh/Source/EditableMesh/EditableMesh.cpp
@@ -653,8 +653,17 @@ void UEditableMesh::SetVertexAttribute( const FVertexID VertexID, const FName At
 		{
 			static TArray<FPolygonID> ConnectedPolygons;
 			GetVertexConnectedPolygons( VertexID, /* Out */ ConnectedPolygons );
-			DeletedOctreePolygonIDs.Append( ConnectedPolygons );
-			NewOctreePolygonIDs.Append( ConnectedPolygons );
+			for( const FPolygonID PolygonID : ConnectedPolygons )
+			{
+				// Only mark the polygon to be removed from the octree if it was already in it to begin with
+				bool bAlreadyInSet = false;
+				NewOctreePolygonIDs.Add( PolygonID, /* Out */ &bAlreadyInSet );
+				const bool bIsAlreadyInOctree = !bAlreadyInSet;
+				if( bIsAlreadyInOctree )
+				{
+					DeletedOctreePolygonIDs.Add( PolygonID );
+				}
+			}
 		}
 	}
 	else if( AttributeName == UEditableMeshAttribute::VertexCornerSharpness() )
@@ -7502,8 +7511,14 @@ void UEditableMesh::InsertPolygonPerimeterVertices( const FPolygonID PolygonID,
 	// Update spatial database
 	if( Octree.IsValid() )
 	{
-		DeletedOctreePolygonIDs.Add( PolygonID );
-		NewOctreePolygonIDs.Add( PolygonID );
+		// Only mark the polygon to be removed from the octree if it was already in it to begin with
+		bool bAlreadyInSet = false;
+		NewOctreePolygonIDs.Add( PolygonID, /* Out */ &bAlreadyInSet );
+		const bool bIsAlreadyInOctree = !bAlreadyInSet;
+		if( bIsAlreadyInOctree )
+		{
+			DeletedOctreePolygonIDs.Add( PolygonID );
+		}
 	}
 
 	UE_LOG( LogEditableMesh, Verbose, TEXT( "* InsertPolygonPerimeterVertices finished" ) );
@@ -7568,8 +7583,14 @@ void UEditableMesh::RemovePolygonPerimeterVertices( const FPolygonID PolygonID,
 	// Update spatial database
 	if( Octree.IsValid() )
 	{
-		DeletedOctreePolygonIDs.Add( PolygonID );
-		NewOctreePolygonIDs.Add( PolygonID );
+		// Only mark the polygon to be removed from the octree if it was already in it to begin with
+		bool bAlreadyInSet = false;
+		NewOctreePolygonIDs.Add( PolygonID, /* Out */ &bAlreadyInSet );
+		const bool bIsAlreadyInOctree = !bAlreadyInSet;
+		if( bIsAlreadyInOctree )
+		{
+			DeletedOctreePolygonIDs.Add( PolygonID );
+		}
 	}
 
 	UE_LOG( LogEditableMesh, Verbose, TEXT( "* RemovePolygonPerimeterVertices finished" ) );
