Mesh Editor: Avoid updating GPU resources for meshes TWICE in a single frame
- When rolling back our preview changes, we now wait until after we've finished making our actual (either interim or final) changes to the mesh that frame before updating GPU resources and component instances
- This change makes dragging actions faster!

diff --git a/Engine/Plugins/Runtime/EditableMesh/Source/EditableMesh/EditableMesh.cpp b/Engine/Plugins/Runtime/EditableMesh/Source/EditableMesh/EditableMesh.cpp
--- a/Engine/Plugins/Runtime/EditableMesh/Source/EditableMesh/EditableMesh.cpp
+++ b/Engine/Plugins/Runtime/EditableMesh/Source/EditableMesh/EditableMesh.cpp
@@ -363,11 +363,10 @@ void UEditableMesh::RebuildRenderMesh()
 {
 	if( !IsBeingModified() )
 	{
-		const bool bRefreshBounds = true;
 		const bool bInvalidateLighting = true;
 		for( UEditableMeshAdapter* Adapter : Adapters )
 		{
-			Adapter->OnRebuildRenderMeshStart( this, bRefreshBounds, bInvalidateLighting );
+			Adapter->OnRebuildRenderMeshStart( this, bInvalidateLighting );
 		}
 	}
 
@@ -378,10 +377,11 @@ void UEditableMesh::RebuildRenderMesh()
 
 	if( !IsBeingModified() )
 	{
-		const bool bUpdateCollision = true;
+		const bool bRebuildBoundsAndCollision = true;
 		for( UEditableMeshAdapter* Adapter : Adapters )
 		{
-			Adapter->OnRebuildRenderMeshFinish( this, bUpdateCollision );
+			const bool bIsPreviewRollback = false;
+			Adapter->OnRebuildRenderMeshFinish( this, bRebuildBoundsAndCollision, bIsPreviewRollback );
 		}
 	}
 }
@@ -413,7 +413,7 @@ void UEditableMesh::StartModification( const EMeshModificationType MeshModificat
 		const bool bInvalidateLighting = ( MeshModificationType == EMeshModificationType::FirstInterim || MeshModificationType == EMeshModificationType::Final );	// @todo mesheditor perf: We can avoid invalidating lighting on 'Final' if we know that a 'FirstInterim' happened since the last 'Final'
 		for( UEditableMeshAdapter* Adapter : Adapters )
 		{
-			Adapter->OnRebuildRenderMeshStart( this, bRefreshBounds, bInvalidateLighting );
+			Adapter->OnRebuildRenderMeshStart( this, bInvalidateLighting );
 		}
 
 		for( UEditableMeshAdapter* Adapter : Adapters )
@@ -496,7 +496,9 @@ void UEditableMesh::EndModification( const bool bFromUndo )
 
 		for( UEditableMeshAdapter* Adapter : Adapters )
 		{
-			Adapter->OnRebuildRenderMeshFinish( this, CurrentModificationType == EMeshModificationType::Final );
+			const bool bRebuildBoundsAndCollision = ( CurrentModificationType == EMeshModificationType::Final );
+			const bool bIsPreviewRollback = !( CurrentModificationType == EMeshModificationType::Final || !bFromUndo );
+			Adapter->OnRebuildRenderMeshFinish( this, bRebuildBoundsAndCollision, bIsPreviewRollback );
 		}
 
 		for( UEditableMeshAdapter* Adapter : Adapters )
diff --git a/Engine/Plugins/Runtime/EditableMesh/Source/EditableMesh/EditableStaticMeshAdapter.cpp b/Engine/Plugins/Runtime/EditableMesh/Source/EditableMesh/EditableStaticMeshAdapter.cpp
--- a/Engine/Plugins/Runtime/EditableMesh/Source/EditableMesh/EditableStaticMeshAdapter.cpp
+++ b/Engine/Plugins/Runtime/EditableMesh/Source/EditableMesh/EditableStaticMeshAdapter.cpp
@@ -799,18 +799,25 @@ void UEditableStaticMeshAdapter::OnStartModification( const UEditableMesh* Edita
 }
 
 
-void UEditableStaticMeshAdapter::OnRebuildRenderMeshStart( const UEditableMesh* EditableMesh, const bool bRefreshBounds, const bool bInvalidateLighting )
+void UEditableStaticMeshAdapter::OnRebuildRenderMeshStart( const UEditableMesh* EditableMesh, const bool bInvalidateLighting )
 {
-	// We're changing the mesh itself, so ALL static mesh components in the scene will need
-	// to be unregistered for this (and reregistered afterwards.)
-	RecreateRenderStateContext = MakeShareable( new FStaticMeshComponentRecreateRenderStateContext( StaticMesh, bInvalidateLighting, bRefreshBounds ) );
+	// We may already have a lock on the rendering resources, if it wasn't released the last time we called EndModification()
+	// on the mesh.  This is only the case when rolling back preview changes for a mesh, because we're guaranteed to apply another
+	// modification to the same mesh in the very same frame.  So we can avoid having to update the GPU resources twice in one frame.
+	if( !RecreateRenderStateContext.IsValid() )
+	{
+		// We're changing the mesh itself, so ALL static mesh components in the scene will need
+		// to be unregistered for this (and reregistered afterwards.)
+		const bool bRefreshBounds = true;
+		RecreateRenderStateContext = MakeShareable( new FStaticMeshComponentRecreateRenderStateContext( StaticMesh, bInvalidateLighting, bRefreshBounds ) );
 
-	// Release the static mesh's resources.
-	StaticMesh->ReleaseResources();
+		// Release the static mesh's resources.
+		StaticMesh->ReleaseResources();
 
-	// Flush the resource release commands to the rendering thread to ensure that the build doesn't occur while a resource is still
-	// allocated, and potentially accessing the UStaticMesh.
-	StaticMesh->ReleaseResourcesFence.Wait();
+		// Flush the resource release commands to the rendering thread to ensure that the build doesn't occur while a resource is still
+		// allocated, and potentially accessing the UStaticMesh.
+		StaticMesh->ReleaseResourcesFence.Wait();
+	}
 }
 
 
@@ -820,22 +827,26 @@ void UEditableStaticMeshAdapter::OnEndModification( const UEditableMesh* Editabl
 }
 
 
-void UEditableStaticMeshAdapter::OnRebuildRenderMeshFinish( const UEditableMesh* EditableMesh, const bool bUpdateCollision )
+void UEditableStaticMeshAdapter::OnRebuildRenderMeshFinish( const UEditableMesh* EditableMesh, const bool bRebuildBoundsAndCollision, const bool bIsPreviewRollback )
 {
-	const bool bShouldRecomputeBounds = bUpdateCollision;	// Only rebuild the bounds from scratch if we've finished a drag (bUpdateCollision will be true)
-	UpdateBounds( EditableMesh, bShouldRecomputeBounds );
+	UpdateBounds( EditableMesh, bRebuildBoundsAndCollision );
 	
-	if( bUpdateCollision )
+	if( bRebuildBoundsAndCollision )
 	{
 		UpdateCollision();
 	}
 
-	StaticMesh->InitResources();
+	// When rolling back preview changes, we'll skip updating GPU resources.  This is because we're guaranteed to get either 'interim' or
+	// 'final' changes to the same mesh later this frame, and we want to avoid updating the GPU resources twice in one frame.
+	if( !bIsPreviewRollback )
+	{
+		StaticMesh->InitResources();
 
-	// NOTE: This can call InvalidateLightingCache() on all components using this mesh, causing Modify() to be 
-	// called on those components!  Just something to be aware of when EndModification() is called within
-	// an undo transaction.
-	RecreateRenderStateContext.Reset();
+		// NOTE: This can call InvalidateLightingCache() on all components using this mesh, causing Modify() to be 
+		// called on those components!  Just something to be aware of when EndModification() is called within
+		// an undo transaction.
+		RecreateRenderStateContext.Reset();
+	}
 }
 
 
diff --git a/Engine/Plugins/Runtime/EditableMesh/Source/EditableMesh/EditableStaticMeshAdapter.h b/Engine/Plugins/Runtime/EditableMesh/Source/EditableMesh/EditableStaticMeshAdapter.h
--- a/Engine/Plugins/Runtime/EditableMesh/Source/EditableMesh/EditableStaticMeshAdapter.h
+++ b/Engine/Plugins/Runtime/EditableMesh/Source/EditableMesh/EditableStaticMeshAdapter.h
@@ -120,9 +120,9 @@ public:
 	void InitEditableStaticMesh( UEditableMesh* EditableMesh, class UPrimitiveComponent& Component, const FEditableMeshSubMeshAddress& InitSubMeshAddress );
 	EDITABLEMESH_API void InitFromBlankStaticMesh( UEditableMesh* EditableMesh, UStaticMesh& InStaticMesh );
 
-	virtual void OnRebuildRenderMeshStart( const UEditableMesh* EditableMesh, const bool bRefreshBounds, const bool bInvalidateLighting ) override;
+	virtual void OnRebuildRenderMeshStart( const UEditableMesh* EditableMesh, const bool bInvalidateLighting ) override;
 	virtual void OnRebuildRenderMesh( const UEditableMesh* EditableMesh ) override;
-	virtual void OnRebuildRenderMeshFinish( const UEditableMesh* EditableMesh, const bool bUpdateCollision ) override;
+	virtual void OnRebuildRenderMeshFinish( const UEditableMesh* EditableMesh, const bool bRebuildBoundsAndCollision, const bool bIsPreviewRollback ) override;
 	virtual void OnStartModification( const UEditableMesh* EditableMesh, const EMeshModificationType MeshModificationType, const EMeshTopologyChange MeshTopologyChange ) override;
 	virtual void OnEndModification( const UEditableMesh* EditableMesh ) override;
 	virtual void OnReindexElements( const UEditableMesh* EditableMesh, const FElementIDRemappings& Remappings ) override;
diff --git a/Engine/Plugins/Runtime/EditableMesh/Source/EditableMesh/Public/EditableMeshAdapter.h b/Engine/Plugins/Runtime/EditableMesh/Source/EditableMesh/Public/EditableMeshAdapter.h
--- a/Engine/Plugins/Runtime/EditableMesh/Source/EditableMesh/Public/EditableMeshAdapter.h
+++ b/Engine/Plugins/Runtime/EditableMesh/Source/EditableMesh/Public/EditableMeshAdapter.h
@@ -16,9 +16,9 @@ class UEditableMeshAdapter : public UObject
 
 public:
 
-	virtual void OnRebuildRenderMeshStart( const UEditableMesh* EditableMesh, const bool bRefreshBounds, const bool bInvalidateLighting ) PURE_VIRTUAL(,);
+	virtual void OnRebuildRenderMeshStart( const UEditableMesh* EditableMesh, const bool bInvalidateLighting ) PURE_VIRTUAL(,);
 	virtual void OnRebuildRenderMesh( const UEditableMesh* EditableMesh ) PURE_VIRTUAL(,);
-	virtual void OnRebuildRenderMeshFinish( const UEditableMesh* EditableMesh, const bool bUpdateCollision ) PURE_VIRTUAL(,);
+	virtual void OnRebuildRenderMeshFinish( const UEditableMesh* EditableMesh, const bool bRebuildBoundsAndCollision, const bool bIsPreviewRollback ) PURE_VIRTUAL(,);
 	virtual void OnStartModification( const UEditableMesh* EditableMesh, const EMeshModificationType MeshModificationType, const EMeshTopologyChange MeshTopologyChange ) PURE_VIRTUAL(,);
 	virtual void OnEndModification( const UEditableMesh* EditableMesh ) PURE_VIRTUAL(,);
 	virtual void OnReindexElements( const UEditableMesh* EditableMesh, const FElementIDRemappings& Remappings ) PURE_VIRTUAL(,);
