Added new functionality to initialize adapters from an editable mesh. Normally, adapters are able to update themselves incrementally as changes are applied to the mesh, but if an existing editable mesh has new adapters added to it, this will need initializing to keep it in line with the mesh.

Added new method to the UEditableMeshAdapter interface, InitializeFromEditableMesh, and a public call on EditableMesh itself to initialize all of its adapters.

diff --git a/Engine/Plugins/Editor/MeshEditor/Source/MeshEditor/MeshEditorStaticMeshAdapter.h b/Engine/Plugins/Editor/MeshEditor/Source/MeshEditor/MeshEditorStaticMeshAdapter.h
--- a/Engine/Plugins/Editor/MeshEditor/Source/MeshEditor/MeshEditorStaticMeshAdapter.h
+++ b/Engine/Plugins/Editor/MeshEditor/Source/MeshEditor/MeshEditorStaticMeshAdapter.h
@@ -22,6 +22,7 @@ public:
 	/** Creates a editable static mesh from the specified component and sub-mesh address */
 	void Initialize( UEditableMesh* EditableMesh, class UWireframeMesh* WireframeMesh );
 
+	virtual void InitializeFromEditableMesh( const UEditableMesh* EditableMesh ) override {}
 	virtual void OnRebuildRenderMeshStart( const UEditableMesh* EditableMesh, const bool bInvalidateLighting ) override;
 	virtual void OnRebuildRenderMesh( const UEditableMesh* EditableMesh ) override;
 	virtual void OnRebuildRenderMeshFinish( const UEditableMesh* EditableMesh, const bool bRebuildBoundsAndCollision, const bool bIsPreviewRollback ) override;
diff --git a/Engine/Plugins/Editor/MeshEditor/Source/MeshEditor/MeshEditorSubdividedStaticMeshAdapter.h b/Engine/Plugins/Editor/MeshEditor/Source/MeshEditor/MeshEditorSubdividedStaticMeshAdapter.h
--- a/Engine/Plugins/Editor/MeshEditor/Source/MeshEditor/MeshEditorSubdividedStaticMeshAdapter.h
+++ b/Engine/Plugins/Editor/MeshEditor/Source/MeshEditor/MeshEditorSubdividedStaticMeshAdapter.h
@@ -22,6 +22,7 @@ public:
 	/** Creates a editable static mesh from the specified component and sub-mesh address */
 	void Initialize( UEditableMesh* EditableMesh, class UWireframeMesh* WireframeMesh );
 
+	virtual void InitializeFromEditableMesh( const UEditableMesh* EditableMesh ) override {}
 	virtual void OnRebuildRenderMeshStart( const UEditableMesh* EditableMesh, const bool bInvalidateLighting ) override;
 	virtual void OnRebuildRenderMesh( const UEditableMesh* EditableMesh ) override;
 	virtual void OnRebuildRenderMeshFinish( const UEditableMesh* EditableMesh, const bool bRebuildBoundsAndCollision, const bool bIsPreviewRollback ) override;
diff --git a/Engine/Plugins/Runtime/EditableMesh/Source/EditableMesh/EditableMesh.cpp b/Engine/Plugins/Runtime/EditableMesh/Source/EditableMesh/EditableMesh.cpp
--- a/Engine/Plugins/Runtime/EditableMesh/Source/EditableMesh/EditableMesh.cpp
+++ b/Engine/Plugins/Runtime/EditableMesh/Source/EditableMesh/EditableMesh.cpp
@@ -361,6 +361,17 @@ void UEditableMesh::Uncompact( const FElementIDRemappings& Remappings )
 }
 
 
+void UEditableMesh::InitializeAdapters()
+{
+	for( UEditableMeshAdapter* Adapter : Adapters )
+	{
+		Adapter->InitializeFromEditableMesh( this );
+	}
+
+	RebuildRenderMesh();
+}
+
+
 void UEditableMesh::RebuildRenderMesh()
 {
 	if( !IsBeingModified() )
diff --git a/Engine/Plugins/Runtime/EditableMesh/Source/EditableMesh/EditableStaticMeshAdapter.cpp b/Engine/Plugins/Runtime/EditableMesh/Source/EditableMesh/EditableStaticMeshAdapter.cpp
--- a/Engine/Plugins/Runtime/EditableMesh/Source/EditableMesh/EditableStaticMeshAdapter.cpp
+++ b/Engine/Plugins/Runtime/EditableMesh/Source/EditableMesh/EditableStaticMeshAdapter.cpp
@@ -502,6 +502,58 @@ void UEditableStaticMeshAdapter::InitFromBlankStaticMesh( UEditableMesh* Editabl
 	StaticMesh = &InStaticMesh;
 }
 
+
+void UEditableStaticMeshAdapter::InitializeFromEditableMesh( const UEditableMesh* EditableMesh )
+{
+	// Get the static mesh from the editable mesh submesh address
+	const FEditableMeshSubMeshAddress& SubMeshAddress = EditableMesh->GetSubMeshAddress();
+	StaticMesh = static_cast<UStaticMesh*>( SubMeshAddress.MeshObjectPtr );
+
+	// @todo mesheditor instancing: sort this out
+	OriginalStaticMesh = nullptr;
+
+	// Always targets LOD0 at the moment
+	StaticMeshLODIndex = 0;
+
+	RenderingPolygons.Reset();
+	RenderingPolygonGroups.Reset();
+
+	// Create all the required rendering polygon groups (initialized to 'empty', each with a unique rendering section index)
+	int32 RenderingSectionIndex = 0;
+	for( TSparseArray<FMeshPolygonGroup>::TConstIterator It( EditableMesh->PolygonGroups ); It; ++It )
+	{
+		const int32 Index = It.GetIndex();
+		RenderingPolygonGroups.Insert( Index, FRenderingPolygonGroup() );
+		FRenderingPolygonGroup& RenderingPolygonGroup = RenderingPolygonGroups[ Index ];
+		RenderingPolygonGroup.RenderingSectionIndex = RenderingSectionIndex;
+		RenderingPolygonGroup.MaxTriangles = 0;
+
+		RenderingSectionIndex++;
+	}
+
+	// Go through all the polygons, adding their triangles to the rendering polygon group
+	for( TSparseArray<FMeshPolygon>::TConstIterator It( EditableMesh->Polygons ); It; ++It )
+	{
+		const int32 Index = It.GetIndex();
+		RenderingPolygons.Insert( Index, FRenderingPolygon() );
+		FRenderingPolygon& RenderingPolygon = RenderingPolygons[ Index ];
+
+		const FPolygonGroupID PolygonGroupID( It->PolygonGroupID );
+		const FMeshPolygonGroup& PolygonGroup = EditableMesh->PolygonGroups[ PolygonGroupID.GetValue() ];
+		FRenderingPolygonGroup& RenderingPolygonGroup = RenderingPolygonGroups[ PolygonGroupID.GetValue() ];
+
+		for( const FMeshTriangle& Triangle : It->Triangles )
+		{
+			const int32 TriangleIndex = RenderingPolygonGroup.Triangles.Add( Triangle );
+			RenderingPolygon.TriangulatedPolygonTriangleIndices.Add( FTriangleID( TriangleIndex ) );
+		}
+
+		RenderingPolygonGroup.MaxTriangles += It->Triangles.Num();
+	}
+
+}
+
+
 int32 GetStaticMeshMaterialIndex( const UStaticMesh* StaticMesh, const UMaterialInterface* MaterialInterface )
 {
 	return StaticMesh->StaticMaterials.IndexOfByPredicate(
diff --git a/Engine/Plugins/Runtime/EditableMesh/Source/EditableMesh/EditableStaticMeshAdapter.h b/Engine/Plugins/Runtime/EditableMesh/Source/EditableMesh/EditableStaticMeshAdapter.h
--- a/Engine/Plugins/Runtime/EditableMesh/Source/EditableMesh/EditableStaticMeshAdapter.h
+++ b/Engine/Plugins/Runtime/EditableMesh/Source/EditableMesh/EditableStaticMeshAdapter.h
@@ -43,12 +43,8 @@ struct FRenderingPolygon
 {
 	GENERATED_BODY()
 
-	/** Static meshes currently only support triangles.  We'll always triangulate polygons and keep track
-	    of all of the triangles here */
-	// @todo mesheditor: For other mesh formats, we may not need to assume triangles.  For example, modern
-	//		GPUs can render quads.  We still need to consider collision geometry though.
-	// @todo mesheditor: This should probably always be a transient thing, right?  We can re-triangulate on 
-	//		load?  Not sure.  We might want some edit-time control over triangulation, which should be saved/loaded.
+	/** This is a list of indices of triangles in the FRenderingPolygonGroup::Triangles array.
+	    We use this to maintain a record of which triangles in the section belong to this polygon. */
 	UPROPERTY()
 	TArray<FTriangleID> TriangulatedPolygonTriangleIndices;
 
@@ -120,6 +116,7 @@ public:
 	void InitEditableStaticMesh( UEditableMesh* EditableMesh, class UPrimitiveComponent& Component, const FEditableMeshSubMeshAddress& InitSubMeshAddress );
 	EDITABLEMESH_API void InitFromBlankStaticMesh( UEditableMesh* EditableMesh, UStaticMesh& InStaticMesh );
 
+	virtual void InitializeFromEditableMesh( const UEditableMesh* EditableMesh ) override;
 	virtual void OnRebuildRenderMeshStart( const UEditableMesh* EditableMesh, const bool bInvalidateLighting ) override;
 	virtual void OnRebuildRenderMesh( const UEditableMesh* EditableMesh ) override;
 	virtual void OnRebuildRenderMeshFinish( const UEditableMesh* EditableMesh, const bool bRebuildBoundsAndCollision, const bool bIsPreviewRollback ) override;
diff --git a/Engine/Plugins/Runtime/EditableMesh/Source/EditableMesh/Public/EditableMesh.h b/Engine/Plugins/Runtime/EditableMesh/Source/EditableMesh/Public/EditableMesh.h
--- a/Engine/Plugins/Runtime/EditableMesh/Source/EditableMesh/Public/EditableMesh.h
+++ b/Engine/Plugins/Runtime/EditableMesh/Source/EditableMesh/Public/EditableMesh.h
@@ -478,6 +478,7 @@ public:
 	/** Remaps mesh element arrays according to the provided remappings, in order to undo a compact operation */
 	void Uncompact( const FElementIDRemappings& Remappings );
 
+	UFUNCTION( BlueprintCallable, Category="Editable Mesh" ) void InitializeAdapters();
 	UFUNCTION( BlueprintCallable, Category="Editable Mesh" ) void RebuildRenderMesh();
 	UFUNCTION( BlueprintCallable, Category="Editable Mesh" ) void StartModification( const EMeshModificationType MeshModificationType, const EMeshTopologyChange MeshTopologyChange );
 	UFUNCTION( BlueprintCallable, Category="Editable Mesh" ) void EndModification( const bool bFromUndo = false );
@@ -908,7 +909,7 @@ public:
 	int32 PendingCompactCounter;
 
 	/** Data will be compacted after this many topology modifying actions. */
-	static const int32 CompactFrequency = 5;
+	static const int32 CompactFrequency = 50;
 
 	/** OpenSubdiv topology refiner object.  This is generated for meshes that have subdivision levels, and reused to generate new limit surfaces 
 	    when geometry is moved.  When the mesh's topology changes, this object is regenerated from scratch. */
diff --git a/Engine/Plugins/Runtime/EditableMesh/Source/EditableMesh/Public/EditableMeshAdapter.h b/Engine/Plugins/Runtime/EditableMesh/Source/EditableMesh/Public/EditableMeshAdapter.h
--- a/Engine/Plugins/Runtime/EditableMesh/Source/EditableMesh/Public/EditableMeshAdapter.h
+++ b/Engine/Plugins/Runtime/EditableMesh/Source/EditableMesh/Public/EditableMeshAdapter.h
@@ -16,6 +16,7 @@ class EDITABLEMESH_API UEditableMeshAdapter : public UObject
 
 public:
 
+	virtual void InitializeFromEditableMesh( const UEditableMesh* EditableMesh ) PURE_VIRTUAL(,);
 	virtual void OnRebuildRenderMeshStart( const UEditableMesh* EditableMesh, const bool bInvalidateLighting ) PURE_VIRTUAL(,);
 	virtual void OnRebuildRenderMesh( const UEditableMesh* EditableMesh ) PURE_VIRTUAL(,);
 	virtual void OnRebuildRenderMeshFinish( const UEditableMesh* EditableMesh, const bool bRebuildBoundsAndCollision, const bool bIsPreviewRollback ) PURE_VIRTUAL(,);
