- Split off mesh description into its own class and engine module. The intention is this will be shared by other mesh objects as the source mesh format, rather than being part of UEditableMesh.

diff --git a/Engine/Plugins/Runtime/EditableMesh/Source/EditableMesh/EditableMesh.cpp b/Engine/Plugins/Runtime/EditableMesh/Source/EditableMesh/EditableMesh.cpp
--- a/Engine/Plugins/Runtime/EditableMesh/Source/EditableMesh/EditableMesh.cpp
+++ b/Engine/Plugins/Runtime/EditableMesh/Source/EditableMesh/EditableMesh.cpp
@@ -52,26 +52,6 @@ namespace EditableMesh
 // =========================================================
 //
 
-const FElementID FElementID::Invalid( TNumericLimits<uint32>::Max() );
-const FVertexID FVertexID::Invalid( TNumericLimits<uint32>::Max() );
-const FVertexInstanceID FVertexInstanceID::Invalid( TNumericLimits<uint32>::Max() );
-const FEdgeID FEdgeID::Invalid( TNumericLimits<uint32>::Max() );
-const FPolygonGroupID FPolygonGroupID::Invalid( TNumericLimits<uint32>::Max() );
-const FPolygonID FPolygonID::Invalid( TNumericLimits<uint32>::Max() );
-
-const FName UEditableMeshAttribute::VertexPositionName( "VertexPosition" );
-const FName UEditableMeshAttribute::VertexCornerSharpnessName( "VertexCornerSharpness" );
-const FName UEditableMeshAttribute::VertexNormalName( "VertexNormal" );
-const FName UEditableMeshAttribute::VertexTangentName( "VertexTangent" );
-const FName UEditableMeshAttribute::VertexBinormalSignName( "VertexBinormalSign" );
-const FName UEditableMeshAttribute::VertexTextureCoordinateName( "VertexTextureCoordinate" );
-const FName UEditableMeshAttribute::VertexColorName( "VertexColor" );
-const FName UEditableMeshAttribute::EdgeIsHardName( "EdgeIsHard" );
-const FName UEditableMeshAttribute::EdgeCreaseSharpnessName( "EdgeCreaseSharpness" );
-const FName UEditableMeshAttribute::PolygonNormalName( "PolygonNormal" );
-const FName UEditableMeshAttribute::PolygonCenterName( "PolygonCenter" );
-
-
 UEditableMesh::UEditableMesh()
 	: bAllowUndo( false ),
 	  bAllowCompact( false ),
@@ -85,12 +65,6 @@ void UEditableMesh::Serialize( FArchive& Ar )
 {
 	Super::Serialize( Ar );
 	Ar.UsingCustomVersion( FEditableMeshCustomVersion::GUID );
-
-	SerializeSparseArray( Ar, Vertices );
-	SerializeSparseArray( Ar, VertexInstances );
-	SerializeSparseArray( Ar, Edges );
-	SerializeSparseArray( Ar, Polygons );
-	SerializeSparseArray( Ar, PolygonGroups );
 }
 
 
@@ -186,116 +160,35 @@ private:
 };
 
 
-void UEditableMesh::FixUpElementIDs( const FElementIDRemappings& Remappings )
+void UEditableMesh::RemapOctreeIDs( const FElementIDRemappings& Remappings )
 {
-	for( FMeshVertex& Vertex : Vertices )
-	{
-		// Fix up vertex instance index references in vertices array
-		for( FVertexInstanceID& VertexInstanceID : Vertex.VertexInstanceIDs )
-		{
-			VertexInstanceID = Remappings.GetRemappedVertexInstanceID( VertexInstanceID );
-		}
-
-		// Fix up edge index references in the vertex array
-		for( FEdgeID& EdgeID : Vertex.ConnectedEdgeIDs )
-		{
-			EdgeID = Remappings.GetRemappedEdgeID( EdgeID );
-		}
-	}
-
-	// Fix up vertex index references in vertex instance array
-	for( FMeshVertexInstance& VertexInstance : VertexInstances )
+	if( Octree.IsValid() )
 	{
-		VertexInstance.VertexID = Remappings.GetRemappedVertexID( VertexInstance.VertexID );
-
-		for( FPolygonID& PolygonID : VertexInstance.ConnectedPolygons )
+		for( FPolygonID& PolygonID : DeletedOctreePolygonIDs )
 		{
 			PolygonID = Remappings.GetRemappedPolygonID( PolygonID );
 		}
-	}
-
-	for( FMeshEdge& Edge : Edges )
-	{
-		// Fix up vertex index references in Edges array
-		for( int32 Index = 0; Index < 2; Index++ )
-		{
-			Edge.VertexIDs[ Index ] = Remappings.GetRemappedVertexID( Edge.VertexIDs[ Index ] );
-		}
-
-		// Fix up references to section indices
-		for( FPolygonID& ConnectedPolygon : Edge.ConnectedPolygons )
-		{
-			ConnectedPolygon = Remappings.GetRemappedPolygonID( ConnectedPolygon );
-		}
-	}
-
-	for( FMeshPolygon& Polygon : Polygons )
-	{
-		// Fix up references to vertex indices in section polygons' contours
-		for( FVertexInstanceID& VertexInstanceID : Polygon.PerimeterContour.VertexInstanceIDs )
-		{
-			VertexInstanceID = Remappings.GetRemappedVertexInstanceID( VertexInstanceID );
-		}
-
-		for( FMeshPolygonContour& HoleContour : Polygon.HoleContours )
-		{
-			for( FVertexInstanceID& VertexInstanceID : HoleContour.VertexInstanceIDs )
-			{
-				VertexInstanceID = Remappings.GetRemappedVertexInstanceID( VertexInstanceID );
-			}
-		}
-
-		for( FMeshTriangle& Triangle : Polygon.Triangles )
-		{
-			for( int32 TriangleVertexNumber = 0; TriangleVertexNumber < 3; ++TriangleVertexNumber )
-			{
-				const FVertexInstanceID OriginalVertexInstanceID = Triangle.GetVertexInstanceID( TriangleVertexNumber );
-				const FVertexInstanceID NewVertexInstanceID = Remappings.GetRemappedVertexInstanceID( OriginalVertexInstanceID );
-				Triangle.SetVertexInstanceID( TriangleVertexNumber, NewVertexInstanceID );
-			}
-		}
 
-		Polygon.PolygonGroupID = Remappings.GetRemappedPolygonGroupID( Polygon.PolygonGroupID );
-	}
-
-	for( FMeshPolygonGroup& PolygonGroup : PolygonGroups )
-	{
-		for( FPolygonID& Polygon : PolygonGroup.Polygons )
+		for( FPolygonID& PolygonID : NewOctreePolygonIDs )
 		{
-			Polygon = Remappings.GetRemappedPolygonID( Polygon );
+			PolygonID = Remappings.GetRemappedPolygonID( PolygonID );
 		}
 	}
 
-	// Fix up spatial database
 	{
-		if( Octree.IsValid() )
-		{
-			for( FPolygonID& PolygonID : DeletedOctreePolygonIDs )
-			{
-				PolygonID = Remappings.GetRemappedPolygonID( PolygonID );
-			}
-
-			for( FPolygonID& PolygonID : NewOctreePolygonIDs )
-			{
-				PolygonID = Remappings.GetRemappedPolygonID( PolygonID );
-			}
-		}
+		// Make a temporary copy of the original map, and clear the stored version
+		static TMap<FPolygonID, FOctreeElementId> OldPolygonIDToOctreeElementIDMap;
+		OldPolygonIDToOctreeElementIDMap = PolygonIDToOctreeElementIDMap;
+		PolygonIDToOctreeElementIDMap.Reset();
 
+		// Rebuild the map with the remapped polygon IDs
+		for( const auto& OldPolygonIDAndOctreeElementID : OldPolygonIDToOctreeElementIDMap )
 		{
-			// Make a temporary copy of the original map, and clear the stored version
-			static TMap<FPolygonID, FOctreeElementId> OldPolygonIDToOctreeElementIDMap;
-			OldPolygonIDToOctreeElementIDMap = PolygonIDToOctreeElementIDMap;
-			PolygonIDToOctreeElementIDMap.Reset();
+			const FPolygonID OldPolygonID = OldPolygonIDAndOctreeElementID.Key;
+			const FOctreeElementId OctreeElementID = OldPolygonIDAndOctreeElementID.Value;
 
-			// Rebuild the map with the remapped polygon IDs
-			for( const auto& OldPolygonIDAndOctreeElementID : OldPolygonIDToOctreeElementIDMap )
-			{
-				const FPolygonID OldPolygonID = OldPolygonIDAndOctreeElementID.Key;
-				const FOctreeElementId OctreeElementID = OldPolygonIDAndOctreeElementID.Value;
-
-				const FPolygonID NewPolygonID = Remappings.GetRemappedPolygonID( OldPolygonID );
-				PolygonIDToOctreeElementIDMap.Add( NewPolygonID, OctreeElementID );
-			}
+			const FPolygonID NewPolygonID = Remappings.GetRemappedPolygonID( OldPolygonID );
+			PolygonIDToOctreeElementIDMap.Add( NewPolygonID, OctreeElementID );
 		}
 	}
 }
@@ -304,14 +197,8 @@ void UEditableMesh::FixUpElementIDs( const FElementIDRemappings& Remappings )
 void UEditableMesh::Compact()
 {
 	static FElementIDRemappings Remappings;
-
-	CompactSparseArrayElements( Vertices, Remappings.NewVertexIndexLookup );
-	CompactSparseArrayElements( VertexInstances, Remappings.NewVertexInstanceIndexLookup );
-	CompactSparseArrayElements( Edges, Remappings.NewEdgeIndexLookup );
-	CompactSparseArrayElements( Polygons, Remappings.NewPolygonIndexLookup );
-	CompactSparseArrayElements( PolygonGroups, Remappings.NewPolygonGroupIndexLookup );
-
-	FixUpElementIDs( Remappings );
+	GetMeshDescription()->Compact( Remappings );
+	RemapOctreeIDs( Remappings );
 
 	for( UEditableMeshAdapter* Adapter : Adapters )
 	{
@@ -338,13 +225,8 @@ void UEditableMesh::Compact()
 
 void UEditableMesh::Uncompact( const FElementIDRemappings& Remappings )
 {
-	RemapSparseArrayElements( Vertices, Remappings.NewVertexIndexLookup );
-	RemapSparseArrayElements( VertexInstances, Remappings.NewVertexInstanceIndexLookup );
-	RemapSparseArrayElements( Edges, Remappings.NewEdgeIndexLookup );
-	RemapSparseArrayElements( Polygons, Remappings.NewPolygonIndexLookup );
-	RemapSparseArrayElements( PolygonGroups, Remappings.NewPolygonGroupIndexLookup );
-
-	FixUpElementIDs( Remappings );
+	GetMeshDescription()->Remap( Remappings );
+	RemapOctreeIDs( Remappings );
 
 	for( UEditableMeshAdapter* Adapter : Adapters )
 	{
@@ -618,271 +500,440 @@ void UEditableMesh::SetSubMeshAddress( const FEditableMeshSubMeshAddress& NewSub
 
 int32 UEditableMesh::GetVertexCount() const
 {
-	return Vertices.Num();
-}
-
-
-int32 UEditableMesh::GetVertexArraySize() const
-{
-	return Vertices.GetMaxIndex();
+	return GetMeshDescription()->Vertices().Num();
 }
 
 
 bool UEditableMesh::IsValidVertex( const FVertexID VertexID ) const
 {
-	return VertexID.GetValue() >= 0 && VertexID.GetValue() < Vertices.GetMaxIndex() && Vertices.IsAllocated( VertexID.GetValue() );
-}
-
-
-FVector4 UEditableMesh::GetVertexAttribute( const FVertexID VertexID, const FName AttributeName, const int32 AttributeIndex ) const
-{
-	const FMeshVertex& Vertex = Vertices[ VertexID.GetValue() ];
-
-	if( AttributeName == UEditableMeshAttribute::VertexPosition() )
-	{
-		check( AttributeIndex == 0 );	// Only one position is supported
-		return FVector4( Vertex.VertexPosition, 0.0f );
-	}
-	else if( AttributeName == UEditableMeshAttribute::VertexCornerSharpness() )
-	{
-		check( AttributeIndex == 0 );	// Only one softness value is supported
-		return FVector4( Vertex.CornerSharpness, 0.0f, 0.0f, 0.0f );
-	}
-
-	checkf( 0, TEXT( "UEditableMesh::GetVertexAttribute() called with unrecognized vertex attribute name: %s (index: %i)" ), *AttributeName.ToString(), AttributeIndex );
-	return FVector4( 0.0f );
+	return GetMeshDescription()->Vertices().IsValid( VertexID );
 }
 
 
-void UEditableMesh::SetVertexAttribute( const FVertexID VertexID, const FName AttributeName, const int32 AttributeIndex, const FVector4 AttributeValue )
+bool UEditableMesh::IsOrphanedVertex( const FVertexID VertexID ) const
 {
-	FMeshVertex& Vertex = Vertices[ VertexID.GetValue() ];
-
-	if( AttributeName == UEditableMeshAttribute::VertexPosition() )
+	const FVertexInstanceArray& VertexInstances = GetMeshDescription()->VertexInstances();
+	for( const FVertexInstanceID VertexInstanceID : GetMeshDescription()->GetVertex( VertexID ).VertexInstanceIDs )
 	{
-		check( AttributeIndex == 0 );	// Only one position is supported
-		Vertex.VertexPosition = AttributeValue;
-
-		// Update spatial database
-		if( Octree.IsValid() )
+		if( VertexInstances[ VertexInstanceID ].ConnectedPolygons.Num() > 0 )
 		{
-			static TArray<FPolygonID> ConnectedPolygons;
-			GetVertexConnectedPolygons( VertexID, /* Out */ ConnectedPolygons );
-			for( const FPolygonID PolygonID : ConnectedPolygons )
-			{
-				// Only mark the polygon to be removed from the octree if it was already in it to begin with
-				bool bAlreadyInSet = false;
-				NewOctreePolygonIDs.Add( PolygonID, /* Out */ &bAlreadyInSet );
-				const bool bIsAlreadyInOctree = !bAlreadyInSet;
-				if( bIsAlreadyInOctree )
-				{
-					DeletedOctreePolygonIDs.Add( PolygonID );
-				}
-			}
+			return false;
 		}
 	}
-	else if( AttributeName == UEditableMeshAttribute::VertexCornerSharpness() )
-	{
-		check( AttributeIndex == 0 );	// Only one softness value is supported
-		Vertex.CornerSharpness = AttributeValue.X;
-	}
-	else
-	{
-		checkf( 0, TEXT( "UEditableMesh::SetVertexAttribute() called with unrecognized vertex attribute name: %s (index: %i)" ), *AttributeName.ToString(), AttributeIndex );
-	}
 
-	for( UEditableMeshAdapter* Adapter : Adapters )
-	{
-		Adapter->OnSetVertexAttribute( this, VertexID, AttributeName, AttributeIndex, AttributeValue );
-	}
+	return true;
 }
 
 
 int32 UEditableMesh::GetVertexConnectedEdgeCount( const FVertexID VertexID ) const
 {
-	checkSlow( Vertices.IsAllocated( VertexID.GetValue() ) );
-	const FMeshVertex& Vertex = Vertices[ VertexID.GetValue() ];
-	return Vertex.ConnectedEdgeIDs.Num();
+	return GetMeshDescription()->GetVertex( VertexID ).ConnectedEdgeIDs.Num();
 }
 
 
 FEdgeID UEditableMesh::GetVertexConnectedEdge( const FVertexID VertexID, const int32 ConnectedEdgeNumber ) const
 {
-	checkSlow( Vertices.IsAllocated( VertexID.GetValue() ) );
-	const FMeshVertex& Vertex = Vertices[ VertexID.GetValue() ];
-	return Vertex.ConnectedEdgeIDs[ ConnectedEdgeNumber ];
+	return GetMeshDescription()->GetVertex( VertexID ).ConnectedEdgeIDs[ ConnectedEdgeNumber ];
 }
 
 
 int32 UEditableMesh::GetVertexInstanceCount() const
 {
-	return VertexInstances.Num();
+	return GetMeshDescription()->VertexInstances().Num();
 }
 
 
-int32 UEditableMesh::GetVertexInstanceArraySize() const
+FVertexID UEditableMesh::GetVertexInstanceVertex( const FVertexInstanceID VertexInstanceID ) const
 {
-	return VertexInstances.GetMaxIndex();
+	return GetMeshDescription()->GetVertexInstance( VertexInstanceID ).VertexID;
 }
 
 
-FVector4 UEditableMesh::GetVertexInstanceAttribute( const FVertexInstanceID VertexInstanceID, const FName AttributeName, const int32 AttributeIndex ) const
+int32 UEditableMesh::GetVertexInstanceConnectedPolygonCount( const FVertexInstanceID VertexInstanceID ) const
 {
-	const FMeshVertexInstance& VertexInstance = VertexInstances[ VertexInstanceID.GetValue() ];
+	return GetMeshDescription()->GetVertexInstance( VertexInstanceID ).ConnectedPolygons.Num();
+}
 
-	if( AttributeName == UEditableMeshAttribute::VertexPosition() )
-	{
-		check( AttributeIndex == 0 );	// Only one position is supported
-		return GetVertexAttribute( VertexInstance.VertexID, AttributeName, AttributeIndex );
-	}
-	else if( AttributeName == UEditableMeshAttribute::VertexNormal() )
-	{
-		check( AttributeIndex == 0 );	// Only one normal is supported
-		return FVector4( VertexInstance.Normal, 0.0f );
-	}
-	else if( AttributeName == UEditableMeshAttribute::VertexTangent() )
-	{
-		check( AttributeIndex == 0 );	// Only one tangent is supported
-		return FVector4( VertexInstance.Tangent, 0.0f );
-	}
-	else if( AttributeName == UEditableMeshAttribute::VertexBinormalSign() )
-	{
-		check( AttributeIndex == 0 );	// Only one basis determinant sign is supported
-		return FVector4( VertexInstance.BinormalSign );
-	}
-	else if( AttributeName == UEditableMeshAttribute::VertexTextureCoordinate() )
-	{
-		const int32 TextureCoordinateIndex = AttributeIndex;
-		if( TextureCoordinateIndex < VertexInstance.VertexUVs.Num() )
-		{
-			const FVector2D TextureCoordinate = VertexInstance.VertexUVs[ TextureCoordinateIndex ];
-			return FVector4( TextureCoordinate, FVector2D::ZeroVector );
-		}
-		else
-		{
-			return FVector4( 0.0f );
-		}
-	}
-	else if( AttributeName == UEditableMeshAttribute::VertexColor() )
-	{
-		check( AttributeIndex == 0 );
-		return FVector4( FLinearColor( VertexInstance.Color ) );
-	}
 
-	checkf( 0, TEXT( "UEditableMesh::GetVertexInstanceAttribute() called with unrecognized vertex instance attribute name: %s (index: %i)" ), *AttributeName.ToString(), AttributeIndex );
-	return FVector4( 0.0f );
+FPolygonID UEditableMesh::GetVertexInstanceConnectedPolygon( const FVertexInstanceID VertexInstanceID, const int32 ConnectedPolygonNumber ) const
+{
+	return GetMeshDescription()->GetVertexInstance( VertexInstanceID ).ConnectedPolygons[ ConnectedPolygonNumber ];
 }
 
 
-void UEditableMesh::SetVertexInstanceAttribute( const FVertexInstanceID VertexInstanceID, const FName AttributeName, const int32 AttributeIndex, const FVector4 AttributeValue )
+int32 UEditableMesh::GetEdgeCount() const
 {
-	FMeshVertexInstance& VertexInstance = VertexInstances[ VertexInstanceID.GetValue() ];
+	return GetMeshDescription()->Edges().Num();
+}
 
-	if( AttributeName == UEditableMeshAttribute::VertexPosition() )
-	{
-		check( AttributeIndex == 0 );	// Only one position is supported
-		SetVertexAttribute( VertexInstance.VertexID, AttributeName, AttributeIndex, AttributeValue );
-	}
-	else if( AttributeName == UEditableMeshAttribute::VertexNormal() )
-	{
-		check( AttributeIndex == 0 );	// Only one normal is supported
-		VertexInstance.Normal = AttributeValue;
-	}
-	else if( AttributeName == UEditableMeshAttribute::VertexTangent() )
-	{
-		check( AttributeIndex == 0 );	// Only one tangent is supported
-		VertexInstance.Tangent = AttributeValue;
-	}
-	else if( AttributeName == UEditableMeshAttribute::VertexBinormalSign() )
-	{
-		check( AttributeIndex == 0 );	// Only one basis determinant sign is supported
-		VertexInstance.BinormalSign = AttributeValue.X;
-	}
-	else if( AttributeName == UEditableMeshAttribute::VertexTextureCoordinate() )
-	{
-		if( AttributeIndex < VertexInstance.VertexUVs.Num() )
-		{
-			VertexInstance.VertexUVs[ AttributeIndex ] = FVector2D( AttributeValue.X, AttributeValue.Y );
-		}
-	}
-	else if( AttributeName == UEditableMeshAttribute::VertexColor() )
-	{
-		check( AttributeIndex == 0 );
-		VertexInstance.Color = FLinearColor( AttributeValue.X, AttributeValue.Y, AttributeValue.Z, AttributeValue.W );
-	}
-	else
-	{
-		checkf( 0, TEXT( "UEditableMesh::SetVertexInstanceAttribute() called with unrecognized vertex instance attribute name: %s (index: %i)" ), *AttributeName.ToString(), AttributeIndex );
-	}
 
-	for( UEditableMeshAdapter* Adapter : Adapters )
-	{
-		Adapter->OnSetVertexInstanceAttribute( this, VertexInstanceID, AttributeName, AttributeIndex, AttributeValue );
-	}
+bool UEditableMesh::IsValidEdge( const FEdgeID EdgeID ) const
+{
+	return GetMeshDescription()->Edges().IsValid( EdgeID );
 }
 
 
-FVector4 UEditableMesh::GetEdgeAttribute( const FEdgeID EdgeID, const FName AttributeName, const int32 AttributeIndex ) const
+FVertexID UEditableMesh::GetEdgeVertex( const FEdgeID EdgeID, const int32 EdgeVertexNumber ) const
 {
-	checkSlow( Edges.IsAllocated( EdgeID.GetValue() ) );
-	const FMeshEdge& Edge = Edges[ EdgeID.GetValue() ];
+	checkSlow( EdgeVertexNumber >= 0 && EdgeVertexNumber < 2 );
+	return GetMeshDescription()->GetEdge( EdgeID ).VertexIDs[ EdgeVertexNumber ];
+}
 
-	if( AttributeName == UEditableMeshAttribute::EdgeIsHard() )
-	{
-		check( AttributeIndex == 0 );	// Only one edge is hard flag is supported
 
-		return FVector4( Edge.bIsHardEdge ? 1.0f : 0.0f, 0.0f, 0.0f, 0.0f );
-	}
-	else if( AttributeName == UEditableMeshAttribute::EdgeCreaseSharpness() )
-	{
-		check( AttributeIndex == 0 );	// Only one edge crease sharpness is supported
+int32 UEditableMesh::GetEdgeConnectedPolygonCount( const FEdgeID EdgeID ) const
+{
+	return GetMeshDescription()->GetEdge( EdgeID ).ConnectedPolygons.Num();
+}
 
-		return FVector4( Edge.CreaseSharpness, 0.0f, 0.0f, 0.0f );
-	}
 
-	checkf( 0, TEXT( "UEditableMesh::GetEdgeAttribute() called with unrecognized edge attribute name: %s (index: %i)" ), *AttributeName.ToString(), AttributeIndex );
-	return FVector4( 0.0f );
+FPolygonID UEditableMesh::GetEdgeConnectedPolygon( const FEdgeID EdgeID, const int32 ConnectedPolygonNumber ) const
+{
+	return GetMeshDescription()->GetEdge( EdgeID ).ConnectedPolygons[ ConnectedPolygonNumber ];
 }
 
 
-void UEditableMesh::SetEdgeAttribute( const FEdgeID EdgeID, const FName AttributeName, const int32 AttributeIndex, const FVector4 AttributeValue, bool bIsUndo )
+int32 UEditableMesh::GetPolygonGroupCount() const
 {
-	checkSlow( Edges.IsAllocated( EdgeID.GetValue() ) );
-	FMeshEdge& Edge = Edges[ EdgeID.GetValue() ];
-
-	if( AttributeName == UEditableMeshAttribute::EdgeIsHard() )
-	{
-		check( AttributeIndex == 0 );	// Only one edge is hard flag is supported
-		SetEdgeHardness( EdgeID, AttributeValue.X != 0.0f, bIsUndo );
+	return GetMeshDescription()->PolygonGroups().Num();
+}
 
-		// @todo mesheditor: this needs to potentially merge/split vertex instances and rebuild normals
-	}
-	else if( AttributeName == UEditableMeshAttribute::EdgeCreaseSharpness() )
-	{
-		check( AttributeIndex == 0 );	// Only one edge crease sharpness is supported
-		Edge.CreaseSharpness = AttributeValue.X;
-	}
-	else
-	{
-		checkf( 0, TEXT( "UEditableMesh::SetEdgeAttribute() called with unrecognized edge attribute name: %s (index: %i)" ), *AttributeName.ToString(), AttributeIndex );
-	}
 
-	for( UEditableMeshAdapter* Adapter : Adapters )
-	{
-		Adapter->OnSetEdgeAttribute( this, EdgeID, AttributeName, AttributeIndex, AttributeValue );
-	}
+bool UEditableMesh::IsValidPolygonGroup( const FPolygonGroupID PolygonGroupID ) const
+{
+	return GetMeshDescription()->PolygonGroups().IsValid( PolygonGroupID );
 }
 
 
-void UEditableMesh::GetPolygonsInSameSoftEdgedGroup( const FVertexID VertexID, const FPolygonID PolygonID, TArray<FPolygonID>& OutPolygonIDs ) const
+int32 UEditableMesh::GetPolygonCountInGroup( const FPolygonGroupID PolygonGroupID ) const
 {
-	// The aim here is to determine which polygons form part of the same soft edged group as the polygons attached to this vertex instance.
-	// They should all contribute to the final vertex instance normal.
+	return GetMeshDescription()->GetPolygonGroup( PolygonGroupID ).Polygons.Num();
+}
 
-	OutPolygonIDs.Reset();
 
-	// Get all polygons connected to this vertex.
-	static TArray<FPolygonID> ConnectedPolygons;
+FPolygonID UEditableMesh::GetPolygonInGroup( const FPolygonGroupID PolygonGroupID, const int32 PolygonNumber ) const
+{
+	return GetMeshDescription()->GetPolygonGroup( PolygonGroupID ).Polygons[ PolygonNumber ];
+}
+
+
+UMaterialInterface* UEditableMesh::GetPolygonGroupMaterial( const FPolygonGroupID PolygonGroupID ) const
+{
+	const FMeshPolygonGroup& PolygonGroup = GetMeshDescription()->GetPolygonGroup( PolygonGroupID );
+	UMaterialInterface* Material = Cast<UMaterialInterface>( PolygonGroup.MaterialAsset.ResolveObject() );
+	check( Material );
+	return Material;
+}
+
+
+FName UEditableMesh::GetPolygonGroupMaterialSlotName( const FPolygonGroupID PolygonGroupID ) const
+{
+	return GetMeshDescription()->GetPolygonGroup( PolygonGroupID ).MaterialSlotName;
+}
+
+
+int32 UEditableMesh::GetPolygonCount() const
+{
+	return GetMeshDescription()->Polygons().Num();
+}
+
+
+bool UEditableMesh::IsValidPolygon( const FPolygonID PolygonID ) const
+{
+	return GetMeshDescription()->Polygons().IsValid( PolygonID );
+}
+
+
+FPolygonGroupID UEditableMesh::GetGroupForPolygon( const FPolygonID PolygonID ) const
+{
+	return GetMeshDescription()->GetPolygon( PolygonID ).PolygonGroupID;
+}
+
+
+int32 UEditableMesh::GetPolygonPerimeterVertexCount( const FPolygonID PolygonID ) const
+{
+	return GetMeshDescription()->GetPolygon( PolygonID ).PerimeterContour.VertexInstanceIDs.Num();
+}
+
+
+FVertexID UEditableMesh::GetPolygonPerimeterVertex( const FPolygonID PolygonID, const int32 PolygonVertexNumber ) const
+{
+	return GetMeshDescription()->GetVertexInstance( GetPolygonPerimeterVertexInstance( PolygonID, PolygonVertexNumber ) ).VertexID;
+}
+
+
+FVertexInstanceID UEditableMesh::GetPolygonPerimeterVertexInstance( const FPolygonID PolygonID, const int32 PolygonVertexNumber ) const
+{
+	return GetMeshDescription()->GetPolygon( PolygonID ).PerimeterContour.VertexInstanceIDs[ PolygonVertexNumber ];
+}
+
+
+FVector4 UEditableMesh::GetPolygonPerimeterVertexAttribute( const FPolygonID PolygonID, const int32 PolygonVertexNumber, const FName AttributeName, const int32 AttributeIndex ) const
+{
+	return GetVertexInstanceAttribute( GetPolygonPerimeterVertexInstance( PolygonID, PolygonVertexNumber ), AttributeName, AttributeIndex );
+}
+
+
+int32 UEditableMesh::GetPolygonHoleCount( const FPolygonID PolygonID ) const
+{
+	return GetMeshDescription()->GetPolygon( PolygonID ).HoleContours.Num();
+}
+
+
+int32 UEditableMesh::GetPolygonHoleVertexCount( const FPolygonID PolygonID, const int32 HoleNumber ) const
+{
+	return GetMeshDescription()->GetPolygon( PolygonID ).HoleContours[ HoleNumber ].VertexInstanceIDs.Num();
+}
+
+
+FVertexID UEditableMesh::GetPolygonHoleVertex( const FPolygonID PolygonID, const int32 HoleNumber, const int32 PolygonVertexNumber ) const
+{
+	return GetMeshDescription()->GetVertexInstance( GetPolygonHoleVertexInstance( PolygonID, HoleNumber, PolygonVertexNumber ) ).VertexID;
+}
+
+
+FVertexInstanceID UEditableMesh::GetPolygonHoleVertexInstance( const FPolygonID PolygonID, const int32 HoleNumber, const int32 PolygonVertexNumber ) const
+{
+	return GetMeshDescription()->GetPolygon( PolygonID ).HoleContours[ HoleNumber ].VertexInstanceIDs[ PolygonVertexNumber ];
+}
+
+
+FVector4 UEditableMesh::GetPolygonHoleVertexAttribute( const FPolygonID PolygonID, const int32 HoleNumber, const int32 PolygonVertexNumber, const FName AttributeName, const int32 AttributeIndex ) const
+{
+	return GetVertexInstanceAttribute( GetPolygonHoleVertexInstance( PolygonID, HoleNumber, PolygonVertexNumber ), AttributeName, AttributeIndex );
+}
+
+
+int32 UEditableMesh::GetPolygonTriangulatedTriangleCount( const FPolygonID PolygonID ) const
+{
+	return GetMeshDescription()->GetPolygon( PolygonID ).Triangles.Num();
+}
+
+
+FMeshTriangle UEditableMesh::GetPolygonTriangulatedTriangle( const FPolygonID PolygonID, int32 PolygonTriangleNumber ) const
+{
+	return GetMeshDescription()->GetPolygon( PolygonID ).Triangles[ PolygonTriangleNumber ];
+}
+
+
+FVector4 UEditableMesh::GetVertexAttribute( const FVertexID VertexID, const FName AttributeName, const int32 AttributeIndex ) const
+{
+	const FMeshVertex& Vertex = GetMeshDescription()->GetVertex( VertexID );
+
+	if( AttributeName == UEditableMeshAttribute::VertexPosition() )
+	{
+		check( AttributeIndex == 0 );	// Only one position is supported
+		return FVector4( Vertex.VertexPosition, 0.0f );
+	}
+	else if( AttributeName == UEditableMeshAttribute::VertexCornerSharpness() )
+	{
+		check( AttributeIndex == 0 );	// Only one softness value is supported
+		return FVector4( Vertex.CornerSharpness, 0.0f, 0.0f, 0.0f );
+	}
+
+	checkf( 0, TEXT( "UEditableMesh::GetVertexAttribute() called with unrecognized vertex attribute name: %s (index: %i)" ), *AttributeName.ToString(), AttributeIndex );
+	return FVector4( 0.0f );
+}
+
+
+FVector4 UEditableMesh::GetVertexInstanceAttribute( const FVertexInstanceID VertexInstanceID, const FName AttributeName, const int32 AttributeIndex ) const
+{
+	const FMeshVertexInstance& VertexInstance = GetMeshDescription()->GetVertexInstance( VertexInstanceID );
+
+	if( AttributeName == UEditableMeshAttribute::VertexPosition() )
+	{
+		check( AttributeIndex == 0 );	// Only one position is supported
+		return GetVertexAttribute( VertexInstance.VertexID, AttributeName, AttributeIndex );
+	}
+	else if( AttributeName == UEditableMeshAttribute::VertexNormal() )
+	{
+		check( AttributeIndex == 0 );	// Only one normal is supported
+		return FVector4( VertexInstance.Normal, 0.0f );
+	}
+	else if( AttributeName == UEditableMeshAttribute::VertexTangent() )
+	{
+		check( AttributeIndex == 0 );	// Only one tangent is supported
+		return FVector4( VertexInstance.Tangent, 0.0f );
+	}
+	else if( AttributeName == UEditableMeshAttribute::VertexBinormalSign() )
+	{
+		check( AttributeIndex == 0 );	// Only one basis determinant sign is supported
+		return FVector4( VertexInstance.BinormalSign );
+	}
+	else if( AttributeName == UEditableMeshAttribute::VertexTextureCoordinate() )
+	{
+		const int32 TextureCoordinateIndex = AttributeIndex;
+		if( TextureCoordinateIndex < VertexInstance.VertexUVs.Num() )
+		{
+			const FVector2D TextureCoordinate = VertexInstance.VertexUVs[ TextureCoordinateIndex ];
+			return FVector4( TextureCoordinate, FVector2D::ZeroVector );
+		}
+		else
+		{
+			return FVector4( 0.0f );
+		}
+	}
+	else if( AttributeName == UEditableMeshAttribute::VertexColor() )
+	{
+		check( AttributeIndex == 0 );
+		return FVector4( FLinearColor( VertexInstance.Color ) );
+	}
+
+	checkf( 0, TEXT( "UEditableMesh::GetVertexInstanceAttribute() called with unrecognized vertex instance attribute name: %s (index: %i)" ), *AttributeName.ToString(), AttributeIndex );
+	return FVector4( 0.0f );
+}
+
+
+FVector4 UEditableMesh::GetEdgeAttribute( const FEdgeID EdgeID, const FName AttributeName, const int32 AttributeIndex ) const
+{
+	const FMeshEdge& Edge = GetMeshDescription()->GetEdge( EdgeID );
+
+	if( AttributeName == UEditableMeshAttribute::EdgeIsHard() )
+	{
+		check( AttributeIndex == 0 );	// Only one edge is hard flag is supported
+
+		return FVector4( Edge.bIsHardEdge ? 1.0f : 0.0f, 0.0f, 0.0f, 0.0f );
+	}
+	else if( AttributeName == UEditableMeshAttribute::EdgeCreaseSharpness() )
+	{
+		check( AttributeIndex == 0 );	// Only one edge crease sharpness is supported
+
+		return FVector4( Edge.CreaseSharpness, 0.0f, 0.0f, 0.0f );
+	}
+
+	checkf( 0, TEXT( "UEditableMesh::GetEdgeAttribute() called with unrecognized edge attribute name: %s (index: %i)" ), *AttributeName.ToString(), AttributeIndex );
+	return FVector4( 0.0f );
+}
+
+
+void UEditableMesh::SetVertexAttribute( const FVertexID VertexID, const FName AttributeName, const int32 AttributeIndex, const FVector4 AttributeValue )
+{
+	FMeshVertex& Vertex = GetMeshDescription()->GetVertex( VertexID );
+
+	if( AttributeName == UEditableMeshAttribute::VertexPosition() )
+	{
+		check( AttributeIndex == 0 );	// Only one position is supported
+		Vertex.VertexPosition = AttributeValue;
+
+		// Update spatial database
+		if( Octree.IsValid() )
+		{
+			static TArray<FPolygonID> ConnectedPolygons;
+			GetVertexConnectedPolygons( VertexID, /* Out */ ConnectedPolygons );
+			for( const FPolygonID PolygonID : ConnectedPolygons )
+			{
+				// Only mark the polygon to be removed from the octree if it was already in it to begin with
+				bool bAlreadyInSet = false;
+				NewOctreePolygonIDs.Add( PolygonID, /* Out */ &bAlreadyInSet );
+				const bool bIsAlreadyInOctree = !bAlreadyInSet;
+				if( bIsAlreadyInOctree )
+				{
+					DeletedOctreePolygonIDs.Add( PolygonID );
+				}
+			}
+		}
+	}
+	else if( AttributeName == UEditableMeshAttribute::VertexCornerSharpness() )
+	{
+		check( AttributeIndex == 0 );	// Only one softness value is supported
+		Vertex.CornerSharpness = AttributeValue.X;
+	}
+	else
+	{
+		checkf( 0, TEXT( "UEditableMesh::SetVertexAttribute() called with unrecognized vertex attribute name: %s (index: %i)" ), *AttributeName.ToString(), AttributeIndex );
+	}
+
+	for( UEditableMeshAdapter* Adapter : Adapters )
+	{
+		Adapter->OnSetVertexAttribute( this, VertexID, AttributeName, AttributeIndex, AttributeValue );
+	}
+}
+
+
+void UEditableMesh::SetVertexInstanceAttribute( const FVertexInstanceID VertexInstanceID, const FName AttributeName, const int32 AttributeIndex, const FVector4 AttributeValue )
+{
+	FMeshVertexInstance& VertexInstance = GetMeshDescription()->GetVertexInstance( VertexInstanceID );
+
+	if( AttributeName == UEditableMeshAttribute::VertexPosition() )
+	{
+		check( AttributeIndex == 0 );	// Only one position is supported
+		SetVertexAttribute( VertexInstance.VertexID, AttributeName, AttributeIndex, AttributeValue );
+	}
+	else if( AttributeName == UEditableMeshAttribute::VertexNormal() )
+	{
+		check( AttributeIndex == 0 );	// Only one normal is supported
+		VertexInstance.Normal = AttributeValue;
+	}
+	else if( AttributeName == UEditableMeshAttribute::VertexTangent() )
+	{
+		check( AttributeIndex == 0 );	// Only one tangent is supported
+		VertexInstance.Tangent = AttributeValue;
+	}
+	else if( AttributeName == UEditableMeshAttribute::VertexBinormalSign() )
+	{
+		check( AttributeIndex == 0 );	// Only one basis determinant sign is supported
+		VertexInstance.BinormalSign = AttributeValue.X;
+	}
+	else if( AttributeName == UEditableMeshAttribute::VertexTextureCoordinate() )
+	{
+		if( AttributeIndex < VertexInstance.VertexUVs.Num() )
+		{
+			VertexInstance.VertexUVs[ AttributeIndex ] = FVector2D( AttributeValue.X, AttributeValue.Y );
+		}
+	}
+	else if( AttributeName == UEditableMeshAttribute::VertexColor() )
+	{
+		check( AttributeIndex == 0 );
+		VertexInstance.Color = FLinearColor( AttributeValue.X, AttributeValue.Y, AttributeValue.Z, AttributeValue.W );
+	}
+	else
+	{
+		checkf( 0, TEXT( "UEditableMesh::SetVertexInstanceAttribute() called with unrecognized vertex instance attribute name: %s (index: %i)" ), *AttributeName.ToString(), AttributeIndex );
+	}
+
+	for( UEditableMeshAdapter* Adapter : Adapters )
+	{
+		Adapter->OnSetVertexInstanceAttribute( this, VertexInstanceID, AttributeName, AttributeIndex, AttributeValue );
+	}
+}
+
+
+void UEditableMesh::SetEdgeAttribute( const FEdgeID EdgeID, const FName AttributeName, const int32 AttributeIndex, const FVector4 AttributeValue, bool bIsUndo )
+{
+	FMeshEdge& Edge = GetMeshDescription()->GetEdge( EdgeID );
+
+	if( AttributeName == UEditableMeshAttribute::EdgeIsHard() )
+	{
+		check( AttributeIndex == 0 );	// Only one edge is hard flag is supported
+		SetEdgeHardness( EdgeID, AttributeValue.X != 0.0f, bIsUndo );
+
+		// @todo mesheditor: this needs to potentially merge/split vertex instances and rebuild normals
+	}
+	else if( AttributeName == UEditableMeshAttribute::EdgeCreaseSharpness() )
+	{
+		check( AttributeIndex == 0 );	// Only one edge crease sharpness is supported
+		Edge.CreaseSharpness = AttributeValue.X;
+	}
+	else
+	{
+		checkf( 0, TEXT( "UEditableMesh::SetEdgeAttribute() called with unrecognized edge attribute name: %s (index: %i)" ), *AttributeName.ToString(), AttributeIndex );
+	}
+
+	for( UEditableMeshAdapter* Adapter : Adapters )
+	{
+		Adapter->OnSetEdgeAttribute( this, EdgeID, AttributeName, AttributeIndex, AttributeValue );
+	}
+}
+
+
+void UEditableMesh::GetPolygonsInSameSoftEdgedGroup( const FVertexID VertexID, const FPolygonID PolygonID, TArray<FPolygonID>& OutPolygonIDs ) const
+{
+	// The aim here is to determine which polygons form part of the same soft edged group as the polygons attached to this vertex instance.
+	// They should all contribute to the final vertex instance normal.
+
+	OutPolygonIDs.Reset();
+
+	// Get all polygons connected to this vertex.
+	static TArray<FPolygonID> ConnectedPolygons;
 	GetVertexConnectedPolygons( VertexID, ConnectedPolygons );
 
 	// Cache a list of all soft edges which share this vertex.
@@ -897,6 +948,7 @@ void UEditableMesh::GetPolygonsInSameSoftEdgedGroup( const FVertexID VertexID, c
 	PolygonsToCheck.Reset();
 	PolygonsToCheck.Add( PolygonID );
 
+	const FEdgeArray& Edges = GetMeshDescription()->Edges();
 	int32 Index = 0;
 	while( Index < PolygonsToCheck.Num() )
 	{
@@ -912,7 +964,7 @@ void UEditableMesh::GetPolygonsInSameSoftEdgedGroup( const FVertexID VertexID, c
 			// have the current polygon as an adjacent.
 			for( const FEdgeID ConnectedSoftEdge : ConnectedSoftEdges )
 			{
-				const FMeshEdge& Edge = Edges[ ConnectedSoftEdge.GetValue() ];
+				const FMeshEdge& Edge = Edges[ ConnectedSoftEdge ];
 				if( Edge.ConnectedPolygons.Contains( PolygonToCheck ) )
 				{
 					for( const FPolygonID AdjacentPolygon : Edge.ConnectedPolygons )
@@ -931,7 +983,7 @@ void UEditableMesh::SplitVertexInstanceInPolygons( const FVertexInstanceID Verte
 {
 	UE_LOG( LogEditableMesh, Verbose, TEXT( "SplitVertexInstanceInPolygons: %s %s" ), *VertexInstanceID.ToString(), *LogHelpers::ArrayToString( PolygonIDs ) );
 
-	const FMeshVertexInstance& VertexInstance = VertexInstances[ VertexInstanceID.GetValue() ];
+	const FMeshVertexInstance& VertexInstance = GetMeshDescription()->GetVertexInstance( VertexInstanceID );
 
 	// Create a new vertex instance copied from the one passed in
 
@@ -971,6 +1023,7 @@ void UEditableMesh::ReplaceVertexInstanceInPolygons( const FVertexInstanceID Old
 	static TArray<FChangeVertexInstancesForPolygon>  VertexInstancesToChange;
 	VertexInstancesToChange.Reset( PolygonIDs.Num() );
 
+	const FPolygonArray& Polygons = GetMeshDescription()->Polygons();
 	for( const FPolygonID PolygonID : PolygonIDs )
 	{
 		VertexInstancesToChange.Emplace();
@@ -978,7 +1031,7 @@ void UEditableMesh::ReplaceVertexInstanceInPolygons( const FVertexInstanceID Old
 
 		ChangeVertexInstances.PolygonID = PolygonID;
 
-		const FMeshPolygon& Polygon = Polygons[ PolygonID.GetValue() ];
+		const FMeshPolygon& Polygon = Polygons[ PolygonID ];
 		int32 VertexInstanceIndex = Polygon.PerimeterContour.VertexInstanceIDs.Find( OldVertexInstanceID );
 
 		if( VertexInstanceIndex != INDEX_NONE )
@@ -1025,7 +1078,7 @@ void UEditableMesh::GetVertexInstanceConnectedPolygonsInSameGroup( const FVertex
 
 	OutPolygonIDs.Reset();
 
-	const FMeshVertexInstance& VertexInstance = VertexInstances[ VertexInstanceID.GetValue() ];
+	const FMeshVertexInstance& VertexInstance = GetMeshDescription()->GetVertexInstance( VertexInstanceID );
 	const FVertexID VertexID = VertexInstance.VertexID;
 
 	// Cache a list of all soft edges which share this vertex.
@@ -1040,6 +1093,7 @@ void UEditableMesh::GetVertexInstanceConnectedPolygonsInSameGroup( const FVertex
 	PolygonsToCheck.Reset();
 	PolygonsToCheck.Add( PolygonID );
 
+	const FEdgeArray& Edges = GetMeshDescription()->Edges();
 	int32 Index = 0;
 	while( Index < PolygonsToCheck.Num() )
 	{
@@ -1055,7 +1109,7 @@ void UEditableMesh::GetVertexInstanceConnectedPolygonsInSameGroup( const FVertex
 			// have the current polygon as an adjacent.
 			for( const FEdgeID ConnectedSoftEdge : ConnectedSoftEdges )
 			{
-				const FMeshEdge& Edge = Edges[ ConnectedSoftEdge.GetValue() ];
+				const FMeshEdge& Edge = Edges[ ConnectedSoftEdge ];
 				if( Edge.ConnectedPolygons.Contains( PolygonToCheck ) )
 				{
 					for( const FPolygonID AdjacentPolygon : Edge.ConnectedPolygons )
@@ -1072,8 +1126,10 @@ void UEditableMesh::GetVertexInstanceConnectedPolygonsInSameGroup( const FVertex
 
 void UEditableMesh::SetEdgeHardness( const FEdgeID EdgeID, const bool bIsHard, const bool bIsUndo )
 {
-	checkSlow( Edges.IsAllocated( EdgeID.GetValue() ) );
-	FMeshEdge& Edge = Edges[ EdgeID.GetValue() ];
+	const FVertexArray& Vertices = GetMeshDescription()->Vertices();
+	const FVertexInstanceArray& VertexInstances = GetMeshDescription()->VertexInstances();
+
+	FMeshEdge& Edge = GetMeshDescription()->GetEdge( EdgeID );
 
 	if( Edge.bIsHardEdge != bIsHard )
 	{
@@ -1089,7 +1145,7 @@ void UEditableMesh::SetEdgeHardness( const FEdgeID EdgeID, const bool bIsHard, c
 
 				for( const FVertexID VertexID : Edge.VertexIDs )
 				{
-					const FMeshVertex& Vertex = Vertices[ VertexID.GetValue() ];
+					const FMeshVertex& Vertex = Vertices[ VertexID ];
 
 					// For each end of the edge, go through all vertex instances, determining if any need to be split.
 					// Take a copy of the vertex instances because splitting them will mutate the list we are iterating.
@@ -1100,306 +1156,119 @@ void UEditableMesh::SetEdgeHardness( const FEdgeID EdgeID, const bool bIsHard, c
 					// We are looking for vertex instances which need splitting into two different instances
 					// (because the hard edge will require separate normals for each of its polygon smoothing groups).
 					for( const FVertexInstanceID VertexInstanceID : VertexInstanceIDsCopy )
-					{
-						const FMeshVertexInstance& VertexInstance = VertexInstances[ VertexInstanceID.GetValue() ];
-
-						// Take a copy of the array of polygons connected to this instance.
-						static TArray<FPolygonID> PolygonsSharingThisVertex;
-						PolygonsSharingThisVertex = VertexInstance.ConnectedPolygons;
-						bool bFirstTime = true;
-						while( PolygonsSharingThisVertex.Num() > 0 )
-						{
-							// For the next polygon in the array, determine all other polygons in the same smoothing group which share this vertex instance.
-							static TArray<FPolygonID> PolygonsInSameSmoothingGroup;
-							GetVertexInstanceConnectedPolygonsInSameGroup( VertexInstanceID, PolygonsSharingThisVertex[ 0 ], PolygonsInSameSmoothingGroup );
-
-							// Check that all polygons in the smoothing group are attached to this vertex instance, and remove them from the master list of polygons
-							// connected to this instance. If a polygon in the smoothing group is not attached to this vertex instance, it's because it's the other
-							// side of a UV seam and hence has a distinct vertex instance.
-							for( const FPolygonID PolygonInSmoothingGroup : PolygonsInSameSmoothingGroup )
-							{
-								verify( PolygonsSharingThisVertex.Remove( PolygonInSmoothingGroup ) == 1 );
-							}
-
-							// First group which we extract: do nothing - they can keep their existing instance ID.
-							// Subsequent times round the loop, we create a new vertex instance copied from the original one, and replace connected polygon vertices with it.
-							if( !bFirstTime )
-							{
-								SplitVertexInstanceInPolygons( VertexInstanceID, PolygonsInSameSmoothingGroup );
-							}
-
-							bFirstTime = false;
-						}
-					}
-				}
-			}
-			else
-			{
-				// Setting a soft edge potentially requires vertex instances to be merged
-
-				static TArray<FVertexInstanceID> VertexInstancesToDelete;
-				VertexInstancesToDelete.Reset();
-
-				for( const FVertexID VertexID : Edge.VertexIDs )
-				{
-					// For each end of the edge which has been softened...
-					const FMeshVertex& Vertex = Vertices[ VertexID.GetValue() ];
-
-					// ...iterate through pairs of vertex instances, looking for potential to merge them
-					for( int32 IndexA = 0; IndexA < Vertex.VertexInstanceIDs.Num() - 1; ++IndexA )
-					{
-						const FVertexInstanceID VertexInstanceIDA = Vertex.VertexInstanceIDs[ IndexA ];
-						const FMeshVertexInstance& VertexInstanceA = VertexInstances[ VertexInstanceIDA.GetValue() ];
-
-						// If vertex instance isn't connected to any polygon, we can't deduce anything about its smoothing group, so skip to the next one
-						if( VertexInstanceA.ConnectedPolygons.Num() == 0 )
-						{
-							continue;
-						}
-
-						// Get mergeable set of attributes for first vertex instance
-						static TArray<FMeshElementAttributeData> AttributesA;
-						GetMergedNamedVertexInstanceAttributeData( GetMergeableVertexInstanceAttributes(), Vertex.VertexInstanceIDs[ IndexA ], TArray<FMeshElementAttributeData>(), AttributesA );
-
-						for( int32 IndexB = IndexA + 1; IndexB < Vertex.VertexInstanceIDs.Num(); ++IndexB )
-						{
-							const FVertexInstanceID VertexInstanceIDB = Vertex.VertexInstanceIDs[ IndexB ];
-
-							// If this vertex instance has been marked for deletion in a previous iteration, skip it
-							if( VertexInstancesToDelete.Contains( VertexInstanceIDB ) )
-							{
-								continue;
-							}
-
-							// Get mergeable set of attributes for second vertex instance
-							static TArray<FMeshElementAttributeData> AttributesB;
-							GetMergedNamedVertexInstanceAttributeData( GetMergeableVertexInstanceAttributes(), Vertex.VertexInstanceIDs[ IndexB ], TArray<FMeshElementAttributeData>(), AttributesB );
-
-							// If both vertex instances have equal mergeable attributes, they can potentially be merged
-							if( AreAttributeListsEqual( AttributesA, AttributesB ) )
-							{
-								const FMeshVertexInstance& VertexInstanceB = VertexInstances[ VertexInstanceIDB.GetValue() ];
-
-								// Check to see if they are in the same smoothing group, using an arbitary connected polygon in the first vertex instance as the pivot point.
-								static TArray<FVertexInstanceID> PossibleVertexInstances;
-								GetVertexInstancesInSameSoftEdgedGroup( VertexInstanceA.VertexID, VertexInstanceA.ConnectedPolygons[ 0 ], false, PossibleVertexInstances );
-
-								// If the smoothing group contains both vertex instances, they can be merged.
-								if( PossibleVertexInstances.Contains( VertexInstanceIDA ) && PossibleVertexInstances.Contains( VertexInstanceIDB ) )
-								{
-									// Change occurrences of VertexInstanceB for VertexInstanceA in VertexInstanceB's connected polygons.
-									// Note, this will cause VertexInstanceA's connected polygons list to be added to (at the end).
-									// This works because we are evaluating the number of connected polygons each time round the loop.
-									ReplaceVertexInstanceInPolygons( VertexInstanceIDB, VertexInstanceIDA, VertexInstanceB.ConnectedPolygons );
-
-									// This will also cause VertexInstanceB to be disconnected from all polygons.
-									// We mark the vertex instance for deletion here, but do not delete it until the end, as to do so would interrupt iterating through vertex instances.
-									check( VertexInstanceB.ConnectedPolygons.Num() == 0 );
-									VertexInstancesToDelete.Add( VertexInstanceIDB );
-								}
-							}
-						}
-					}
-				}
-
-				// Delete orphaned vertex instances
-				const bool bDeleteOrphanedVertices = false;
-				DeleteVertexInstances( VertexInstancesToDelete, bDeleteOrphanedVertices );
-			}
-
-			UE_LOG( LogEditableMesh, Verbose, TEXT( "* SetEdgeHardness returned" ) );
-		}
-
-		// Mark polygons adjacent to this edge as candidates for normal/tangent recalculation
-		PolygonsPendingNewTangentBasis.Append( Edge.ConnectedPolygons );
-	}
-}
-
-
-int32 UEditableMesh::GetEdgeCount() const
-{
-	return Edges.Num();
-}
-
-
-int32 UEditableMesh::GetEdgeArraySize() const
-{
-	return Edges.GetMaxIndex();
-}
-
-
-bool UEditableMesh::IsValidEdge( const FEdgeID EdgeID ) const
-{
-	return EdgeID.GetValue() >= 0 && EdgeID.GetValue() < Edges.GetMaxIndex() && Edges.IsAllocated( EdgeID.GetValue() );
-}
-
-
-FVertexID UEditableMesh::GetEdgeVertex( const FEdgeID EdgeID, const int32 EdgeVertexNumber ) const
-{
-	checkSlow( EdgeVertexNumber >= 0 && EdgeVertexNumber < 2 );
-	checkSlow( Edges.IsAllocated( EdgeID.GetValue() ) );
-	return Edges[ EdgeID.GetValue() ].VertexIDs[ EdgeVertexNumber ];
-}
-
-
-int32 UEditableMesh::GetEdgeConnectedPolygonCount( const FEdgeID EdgeID ) const
-{
-	checkSlow( Edges.IsAllocated( EdgeID.GetValue() ) );
-	const FMeshEdge& Edge = Edges[ EdgeID.GetValue() ];
-	return Edge.ConnectedPolygons.Num();
-}
-
-
-FPolygonID UEditableMesh::GetEdgeConnectedPolygon( const FEdgeID EdgeID, const int32 ConnectedPolygonNumber ) const
-{
-	checkSlow( Edges.IsAllocated( EdgeID.GetValue() ) );
-	const FMeshEdge& Edge = Edges[ EdgeID.GetValue() ];
-	const FPolygonID PolygonID = Edge.ConnectedPolygons[ ConnectedPolygonNumber ];
-	return PolygonID;
-}
-
-
-int32 UEditableMesh::GetPolygonGroupCount() const
-{
-	return PolygonGroups.Num();
-}
-
-
-int32 UEditableMesh::GetPolygonGroupArraySize() const
-{
-	return PolygonGroups.GetMaxIndex();
-}
-
-
-bool UEditableMesh::IsValidPolygonGroup( const FPolygonGroupID PolygonGroupID ) const
-{
-	return
-		PolygonGroupID.GetValue() >= 0 &&
-		PolygonGroupID.GetValue() < PolygonGroups.GetMaxIndex() &&
-		PolygonGroups.IsAllocated( PolygonGroupID.GetValue() );
-}
-
-
-int32 UEditableMesh::GetPolygonCountInGroup( const FPolygonGroupID PolygonGroupID ) const
-{
-	checkSlow( PolygonGroups.IsAllocated( PolygonGroupID.GetValue() ) );
-	return PolygonGroups[ PolygonGroupID.GetValue() ].Polygons.Num();
-}
-
-
-FPolygonID UEditableMesh::GetPolygonInGroup( const FPolygonGroupID PolygonGroupID, const int32 PolygonNumber ) const
-{
-	checkSlow( PolygonGroups.IsAllocated( PolygonGroupID.GetValue() ) );
-	const FMeshPolygonGroup& PolygonGroup = PolygonGroups[ PolygonGroupID.GetValue() ];
-	return FPolygonID( PolygonGroup.Polygons[ PolygonNumber ] );
-}
-
-
-int32 UEditableMesh::GetPolygonCount() const
-{
-	return Polygons.Num();
-}
-
-
-int32 UEditableMesh::GetPolygonArraySize() const
-{
-	return Polygons.GetMaxIndex();
-}
-
-
-bool UEditableMesh::IsValidPolygon( const FPolygonID PolygonID ) const
-{
-	return
-		PolygonID.GetValue() >= 0 &&
-		PolygonID.GetValue() < Polygons.GetMaxIndex() &&
-		Polygons.IsAllocated( PolygonID.GetValue() );
-}
-
-
-FPolygonGroupID UEditableMesh::GetGroupForPolygon( const FPolygonID PolygonID ) const
-{
-	checkSlow( Polygons.IsAllocated( PolygonID.GetValue() ) );
-	return Polygons[ PolygonID.GetValue() ].PolygonGroupID;
-}
-
+					{
+						const FMeshVertexInstance& VertexInstance = VertexInstances[ VertexInstanceID ];
 
-int32 UEditableMesh::GetPolygonPerimeterVertexCount( const FPolygonID PolygonID ) const
-{
-	checkSlow( Polygons.IsAllocated( PolygonID.GetValue() ) );
-	const FMeshPolygon& Polygon = Polygons[ PolygonID.GetValue() ];
+						// Take a copy of the array of polygons connected to this instance.
+						static TArray<FPolygonID> PolygonsSharingThisVertex;
+						PolygonsSharingThisVertex = VertexInstance.ConnectedPolygons;
+						bool bFirstTime = true;
+						while( PolygonsSharingThisVertex.Num() > 0 )
+						{
+							// For the next polygon in the array, determine all other polygons in the same smoothing group which share this vertex instance.
+							static TArray<FPolygonID> PolygonsInSameSmoothingGroup;
+							GetVertexInstanceConnectedPolygonsInSameGroup( VertexInstanceID, PolygonsSharingThisVertex[ 0 ], PolygonsInSameSmoothingGroup );
 
-	return Polygon.PerimeterContour.VertexInstanceIDs.Num();
-}
+							// Check that all polygons in the smoothing group are attached to this vertex instance, and remove them from the master list of polygons
+							// connected to this instance. If a polygon in the smoothing group is not attached to this vertex instance, it's because it's the other
+							// side of a UV seam and hence has a distinct vertex instance.
+							for( const FPolygonID PolygonInSmoothingGroup : PolygonsInSameSmoothingGroup )
+							{
+								verify( PolygonsSharingThisVertex.Remove( PolygonInSmoothingGroup ) == 1 );
+							}
 
+							// First group which we extract: do nothing - they can keep their existing instance ID.
+							// Subsequent times round the loop, we create a new vertex instance copied from the original one, and replace connected polygon vertices with it.
+							if( !bFirstTime )
+							{
+								SplitVertexInstanceInPolygons( VertexInstanceID, PolygonsInSameSmoothingGroup );
+							}
 
-FVertexID UEditableMesh::GetPolygonPerimeterVertex( const FPolygonID PolygonID, const int32 PolygonVertexNumber ) const
-{
-	checkSlow( Polygons.IsAllocated( PolygonID.GetValue() ) );
-	const FMeshPolygon& Polygon = Polygons[ PolygonID.GetValue() ];
-	const FVertexInstanceID VertexInstanceID = Polygon.PerimeterContour.VertexInstanceIDs[ PolygonVertexNumber ];
-	const FMeshVertexInstance& VertexInstance = VertexInstances[ VertexInstanceID.GetValue() ];
-	return VertexInstance.VertexID;
-}
+							bFirstTime = false;
+						}
+					}
+				}
+			}
+			else
+			{
+				// Setting a soft edge potentially requires vertex instances to be merged
 
+				static TArray<FVertexInstanceID> VertexInstancesToDelete;
+				VertexInstancesToDelete.Reset();
 
-FVertexInstanceID UEditableMesh::GetPolygonPerimeterVertexInstance( const FPolygonID PolygonID, const int32 PolygonVertexNumber ) const
-{
-	checkSlow( Polygons.IsAllocated( PolygonID.GetValue() ) );
-	const FMeshPolygon& Polygon = Polygons[ PolygonID.GetValue() ];
-	return Polygon.PerimeterContour.VertexInstanceIDs[ PolygonVertexNumber ];
-}
+				for( const FVertexID VertexID : Edge.VertexIDs )
+				{
+					// For each end of the edge which has been softened...
+					const FMeshVertex& Vertex = Vertices[ VertexID ];
 
+					// ...iterate through pairs of vertex instances, looking for potential to merge them
+					for( int32 IndexA = 0; IndexA < Vertex.VertexInstanceIDs.Num() - 1; ++IndexA )
+					{
+						const FVertexInstanceID VertexInstanceIDA = Vertex.VertexInstanceIDs[ IndexA ];
+						const FMeshVertexInstance& VertexInstanceA = VertexInstances[ VertexInstanceIDA ];
 
-FVector4 UEditableMesh::GetPolygonPerimeterVertexAttribute( const FPolygonID PolygonID, const int32 PolygonVertexNumber, const FName AttributeName, const int32 AttributeIndex ) const
-{
-	checkSlow( Polygons.IsAllocated( PolygonID.GetValue() ) );
-	const FMeshPolygon& Polygon = Polygons[ PolygonID.GetValue() ];
-	const FVertexInstanceID VertexInstanceID = Polygon.PerimeterContour.VertexInstanceIDs[ PolygonVertexNumber ];
-	return GetVertexInstanceAttribute( VertexInstanceID, AttributeName, AttributeIndex );
-}
+						// If vertex instance isn't connected to any polygon, we can't deduce anything about its smoothing group, so skip to the next one
+						if( VertexInstanceA.ConnectedPolygons.Num() == 0 )
+						{
+							continue;
+						}
 
+						// Get mergeable set of attributes for first vertex instance
+						static TArray<FMeshElementAttributeData> AttributesA;
+						GetMergedNamedVertexInstanceAttributeData( GetMergeableVertexInstanceAttributes(), Vertex.VertexInstanceIDs[ IndexA ], TArray<FMeshElementAttributeData>(), AttributesA );
 
-int32 UEditableMesh::GetPolygonHoleCount( const FPolygonID PolygonID ) const
-{
-	checkSlow( Polygons.IsAllocated( PolygonID.GetValue() ) );
-	const FMeshPolygon& Polygon = Polygons[ PolygonID.GetValue() ];
+						for( int32 IndexB = IndexA + 1; IndexB < Vertex.VertexInstanceIDs.Num(); ++IndexB )
+						{
+							const FVertexInstanceID VertexInstanceIDB = Vertex.VertexInstanceIDs[ IndexB ];
 
-	return Polygon.HoleContours.Num();
-}
+							// If this vertex instance has been marked for deletion in a previous iteration, skip it
+							if( VertexInstancesToDelete.Contains( VertexInstanceIDB ) )
+							{
+								continue;
+							}
 
+							// Get mergeable set of attributes for second vertex instance
+							static TArray<FMeshElementAttributeData> AttributesB;
+							GetMergedNamedVertexInstanceAttributeData( GetMergeableVertexInstanceAttributes(), Vertex.VertexInstanceIDs[ IndexB ], TArray<FMeshElementAttributeData>(), AttributesB );
 
-int32 UEditableMesh::GetPolygonHoleVertexCount( const FPolygonID PolygonID, const int32 HoleNumber ) const
-{
-	checkSlow( Polygons.IsAllocated( PolygonID.GetValue() ) );
-	const FMeshPolygon& Polygon = Polygons[ PolygonID.GetValue() ];
-	return Polygon.HoleContours[ HoleNumber ].VertexInstanceIDs.Num();
-}
+							// If both vertex instances have equal mergeable attributes, they can potentially be merged
+							if( AreAttributeListsEqual( AttributesA, AttributesB ) )
+							{
+								const FMeshVertexInstance& VertexInstanceB = VertexInstances[ VertexInstanceIDB ];
 
+								// Check to see if they are in the same smoothing group, using an arbitary connected polygon in the first vertex instance as the pivot point.
+								static TArray<FVertexInstanceID> PossibleVertexInstances;
+								GetVertexInstancesInSameSoftEdgedGroup( VertexInstanceA.VertexID, VertexInstanceA.ConnectedPolygons[ 0 ], false, PossibleVertexInstances );
 
-FVertexID UEditableMesh::GetPolygonHoleVertex( const FPolygonID PolygonID, const int32 HoleNumber, const int32 PolygonVertexNumber ) const
-{
-	checkSlow( Polygons.IsAllocated( PolygonID.GetValue() ) );
-	const FMeshPolygon& Polygon = Polygons[ PolygonID.GetValue() ];
-	const FVertexInstanceID VertexInstanceID = Polygon.HoleContours[ HoleNumber ].VertexInstanceIDs[ PolygonVertexNumber ];
-	const FMeshVertexInstance& VertexInstance = VertexInstances[ VertexInstanceID.GetValue() ];
-	return VertexInstance.VertexID;
-}
+								// If the smoothing group contains both vertex instances, they can be merged.
+								if( PossibleVertexInstances.Contains( VertexInstanceIDA ) && PossibleVertexInstances.Contains( VertexInstanceIDB ) )
+								{
+									// Change occurrences of VertexInstanceB for VertexInstanceA in VertexInstanceB's connected polygons.
+									// Note, this will cause VertexInstanceA's connected polygons list to be added to (at the end).
+									// This works because we are evaluating the number of connected polygons each time round the loop.
+									ReplaceVertexInstanceInPolygons( VertexInstanceIDB, VertexInstanceIDA, VertexInstanceB.ConnectedPolygons );
 
+									// This will also cause VertexInstanceB to be disconnected from all polygons.
+									// We mark the vertex instance for deletion here, but do not delete it until the end, as to do so would interrupt iterating through vertex instances.
+									check( VertexInstanceB.ConnectedPolygons.Num() == 0 );
+									VertexInstancesToDelete.Add( VertexInstanceIDB );
+								}
+							}
+						}
+					}
+				}
 
-FVertexInstanceID UEditableMesh::GetPolygonHoleVertexInstance( const FPolygonID PolygonID, const int32 HoleNumber, const int32 PolygonVertexNumber ) const
-{
-	checkSlow( Polygons.IsAllocated( PolygonID.GetValue() ) );
-	const FMeshPolygon& Polygon = Polygons[ PolygonID.GetValue() ];
-	return Polygon.HoleContours[ HoleNumber ].VertexInstanceIDs[ PolygonVertexNumber ];
-}
+				// Delete orphaned vertex instances
+				const bool bDeleteOrphanedVertices = false;
+				DeleteVertexInstances( VertexInstancesToDelete, bDeleteOrphanedVertices );
+			}
 
+			UE_LOG( LogEditableMesh, Verbose, TEXT( "* SetEdgeHardness returned" ) );
+		}
 
-FVector4 UEditableMesh::GetPolygonHoleVertexAttribute( const FPolygonID PolygonID, const int32 HoleNumber, const int32 PolygonVertexNumber, const FName AttributeName, const int32 AttributeIndex ) const
-{
-	checkSlow( Polygons.IsAllocated( PolygonID.GetValue() ) );
-	const FMeshPolygon& Polygon = Polygons[ PolygonID.GetValue() ];
-	const FVertexInstanceID VertexInstanceID = Polygon.HoleContours[ HoleNumber ].VertexInstanceIDs[ PolygonVertexNumber ];
-	return GetVertexInstanceAttribute( VertexInstanceID, AttributeName, AttributeIndex );
+		// Mark polygons adjacent to this edge as candidates for normal/tangent recalculation
+		PolygonsPendingNewTangentBasis.Append( Edge.ConnectedPolygons );
+	}
 }
 
 
@@ -1467,14 +1336,7 @@ int32 UEditableMesh::GetMaxAttributeIndex( const FName AttributeName ) const
 
 FPolygonGroupID UEditableMesh::GetFirstValidPolygonGroup() const
 {
-	FPolygonGroupID FirstValidPolygonGroupID = FPolygonGroupID::Invalid;
-
-	if( PolygonGroups.Num() > 0 )
-	{
-		FirstValidPolygonGroupID = FPolygonGroupID( PolygonGroups.CreateConstIterator().GetIndex() );
-	}
-
-	return FirstValidPolygonGroupID;
+	return GetMeshDescription()->PolygonGroups().GetFirstValidID();
 }
 
 
@@ -1514,31 +1376,17 @@ void UEditableMesh::GetVertexConnectedPolygons( const FVertexID VertexID, TArray
 {
 	OutConnectedPolygonIDs.Reset();
 
-	checkSlow( Vertices.IsAllocated( VertexID.GetValue() ) );
-	const FMeshVertex& Vertex = Vertices[ VertexID.GetValue() ];
-	for( const FVertexInstanceID VertexInstanceID : Vertex.VertexInstanceIDs )
+	const FVertexInstanceArray& VertexInstances = GetMeshDescription()->VertexInstances();
+	for( const FVertexInstanceID VertexInstanceID : GetMeshDescription()->GetVertex( VertexID ).VertexInstanceIDs )
 	{
-		checkSlow( VertexInstances.IsAllocated( VertexInstanceID.GetValue() ) );
-		const FMeshVertexInstance& VertexInstance = VertexInstances[ VertexInstanceID.GetValue() ];
-		OutConnectedPolygonIDs.Append( VertexInstance.ConnectedPolygons );
+		OutConnectedPolygonIDs.Append( VertexInstances[ VertexInstanceID ].ConnectedPolygons );
 	}
 }
 
 
 void UEditableMesh::GetVertexInstanceConnectedPolygons( const FVertexInstanceID VertexInstanceID, TArray<FPolygonID>& OutConnectedPolygonIDs ) const
 {
-	checkSlow( VertexInstances.IsAllocated( VertexInstanceID.GetValue() ) );
-	const FMeshVertexInstance& VertexInstance = VertexInstances[ VertexInstanceID.GetValue() ];
-
-	OutConnectedPolygonIDs = VertexInstance.ConnectedPolygons;
-}
-
-
-int32 UEditableMesh::GetVertexInstanceConnectedPolygonCount( const FVertexInstanceID VertexInstanceID ) const
-{
-	checkSlow( VertexInstances.IsAllocated( VertexInstanceID.GetValue() ) );
-	const FMeshVertexInstance& VertexInstance = VertexInstances[ VertexInstanceID.GetValue() ];
-	return VertexInstance.ConnectedPolygons.Num();
+	OutConnectedPolygonIDs = GetMeshDescription()->GetVertexInstance( VertexInstanceID ).ConnectedPolygons;
 }
 
 
@@ -1546,11 +1394,10 @@ void UEditableMesh::GetVertexAdjacentVertices( const FVertexID VertexID, TArray<
 {
 	OutAdjacentVertexIDs.Reset();
 
-	const FMeshVertex& Vertex = Vertices[ VertexID.GetValue() ];
-	for( const FEdgeID EdgeID : Vertex.ConnectedEdgeIDs )
+	const FEdgeArray& Edges = GetMeshDescription()->Edges();
+	for( const FEdgeID EdgeID : GetMeshDescription()->GetVertex( VertexID ).ConnectedEdgeIDs )
 	{
-		checkSlow( Edges.IsAllocated( EdgeID.GetValue() ) );
-		const FMeshEdge& Edge = Edges[ EdgeID.GetValue() ];
+		const FMeshEdge& Edge = Edges[ EdgeID ];
 		OutAdjacentVertexIDs.Add( Edge.VertexIDs[ 0 ] == VertexID ? Edge.VertexIDs[ 1 ] : Edge.VertexIDs[ 0 ] );
 	}
 }
@@ -1558,8 +1405,7 @@ void UEditableMesh::GetVertexAdjacentVertices( const FVertexID VertexID, TArray<
 
 void UEditableMesh::GetEdgeVertices( const FEdgeID EdgeID, FVertexID& OutEdgeVertexID0, FVertexID& OutEdgeVertexID1 ) const
 {
-	checkSlow( Edges.IsAllocated( EdgeID.GetValue() ) );
-	const FMeshEdge& Edge = Edges[ EdgeID.GetValue() ];
+	const FMeshEdge& Edge = GetMeshDescription()->GetEdge( EdgeID );
 	OutEdgeVertexID0 = Edge.VertexIDs[ 0 ];
 	OutEdgeVertexID1 = Edge.VertexIDs[ 1 ];
 }
@@ -1567,9 +1413,7 @@ void UEditableMesh::GetEdgeVertices( const FEdgeID EdgeID, FVertexID& OutEdgeVer
 
 void UEditableMesh::GetEdgeConnectedPolygons( const FEdgeID EdgeID, TArray<FPolygonID>& OutConnectedPolygonIDs ) const
 {
-	checkSlow( Edges.IsAllocated( EdgeID.GetValue() ) );
-	const FMeshEdge& Edge = Edges[ EdgeID.GetValue() ];
-	OutConnectedPolygonIDs = Edge.ConnectedPolygons;
+	OutConnectedPolygonIDs = GetMeshDescription()->GetEdge( EdgeID ).ConnectedPolygons;
 }
 
 
@@ -1602,6 +1446,9 @@ void UEditableMesh::GetEdgeLoopElements( const FEdgeID EdgeID, TArray<FEdgeID>&
 {
 	EdgeLoopIDs.Reset();
 
+	const FEdgeArray& Edges = GetMeshDescription()->Edges();
+	const FVertexArray& Vertices = GetMeshDescription()->Vertices();
+
 	// Maintain a list of unique edge IDs which form the loop
 	static TSet<FEdgeID> EdgeIDs;
 	EdgeIDs.Reset();
@@ -1623,8 +1470,7 @@ void UEditableMesh::GetEdgeLoopElements( const FEdgeID EdgeID, TArray<FEdgeID>&
 		// Get the edge so we can see which polygons are connected to it.
 		// When continuing the loop, the criterion is that new edges must share no polygons with this edge,
 		// i.e. they are the other side of a perpendicular edge.
-		checkSlow( Edges.IsAllocated( CurrentEdgeID.GetValue() ) );
-		const FMeshEdge& Edge = Edges[ CurrentEdgeID.GetValue() ];
+		const FMeshEdge& Edge = Edges[ CurrentEdgeID ];
 
 		// Now look for edges connected to each end of this edge
 		for( const FVertexID ConnectedVertexID : Edge.VertexIDs )
@@ -1639,7 +1485,7 @@ void UEditableMesh::GetEdgeLoopElements( const FEdgeID EdgeID, TArray<FEdgeID>&
 			FEdgeID AdjacentEdgeID = FEdgeID::Invalid;
 
 			// Iterate through all edges connected to this vertex
-			const FMeshVertex& Vertex = Vertices[ ConnectedVertexID.GetValue() ];
+			const FMeshVertex& Vertex = Vertices[ ConnectedVertexID ];
 			for( const FEdgeID ConnectedEdgeID : Vertex.ConnectedEdgeIDs )
 			{
 				// If this edge hasn't been added to the loop...
@@ -1647,8 +1493,7 @@ void UEditableMesh::GetEdgeLoopElements( const FEdgeID EdgeID, TArray<FEdgeID>&
 				{
 					// ...see if it shares any polygons with the original edge (intersection operation)
 					bool bIsCandidateEdge = true;
-					checkSlow( Edges.IsAllocated( ConnectedEdgeID.GetValue() ) );
-					const FMeshEdge& ConnectedEdge = Edges[ ConnectedEdgeID.GetValue() ];
+					const FMeshEdge& ConnectedEdge = Edges[ ConnectedEdgeID ];
 					for( const FPolygonID ConnectedPolygonID : ConnectedEdge.ConnectedPolygons )
 					{
 						if( Edge.ConnectedPolygons.Contains( ConnectedPolygonID ) )
@@ -1706,14 +1551,15 @@ int32 UEditableMesh::GetPolygonHoleEdgeCount( const FPolygonID PolygonID, const
 
 void UEditableMesh::GetPolygonPerimeterVertices( const FPolygonID PolygonID, TArray<FVertexID>& OutPolygonPerimeterVertexIDs ) const
 {
-	checkSlow( Polygons.IsAllocated( PolygonID.GetValue() ) );
-	const FMeshPolygon& Polygon = Polygons[ PolygonID.GetValue() ];
+	const FVertexInstanceArray& VertexInstances = GetMeshDescription()->VertexInstances();
+	const FMeshPolygon& Polygon = GetMeshDescription()->GetPolygon( PolygonID );
+
 	OutPolygonPerimeterVertexIDs.SetNumUninitialized( Polygon.PerimeterContour.VertexInstanceIDs.Num(), false );
 
 	int32 Index = 0;
 	for( const FVertexInstanceID VertexInstanceID : Polygon.PerimeterContour.VertexInstanceIDs )
 	{
-		const FMeshVertexInstance& VertexInstance = VertexInstances[ VertexInstanceID.GetValue() ];
+		const FMeshVertexInstance& VertexInstance = VertexInstances[ VertexInstanceID ];
 		OutPolygonPerimeterVertexIDs[ Index ] = VertexInstance.VertexID;
 		Index++;
 	}
@@ -1722,16 +1568,15 @@ void UEditableMesh::GetPolygonPerimeterVertices( const FPolygonID PolygonID, TAr
 
 void UEditableMesh::GetPolygonPerimeterVertexInstances( const FPolygonID PolygonID, TArray<FVertexInstanceID>& OutPolygonPerimeterVertexInstanceIDs ) const
 {
-	checkSlow( Polygons.IsAllocated( PolygonID.GetValue() ) );
-	const FMeshPolygon& Polygon = Polygons[ PolygonID.GetValue() ];
+	const FMeshPolygon& Polygon = GetMeshDescription()->GetPolygon( PolygonID );
 	OutPolygonPerimeterVertexInstanceIDs = Polygon.PerimeterContour.VertexInstanceIDs;
 }
 
 
 void UEditableMesh::GetPolygonHoleVertices( const FPolygonID PolygonID, const int32 HoleNumber, TArray<FVertexID>& OutHoleVertexIDs ) const
 {
-	checkSlow( Polygons.IsAllocated( PolygonID.GetValue() ) );
-	const FMeshPolygon& Polygon = Polygons[ PolygonID.GetValue() ];
+	const FVertexInstanceArray& VertexInstances = GetMeshDescription()->VertexInstances();
+	const FMeshPolygon& Polygon = GetMeshDescription()->GetPolygon( PolygonID );
 	const FMeshPolygonContour& Contour = Polygon.HoleContours[ HoleNumber ];
 
 	OutHoleVertexIDs.SetNumUninitialized( Contour.VertexInstanceIDs.Num(), false );
@@ -1739,7 +1584,7 @@ void UEditableMesh::GetPolygonHoleVertices( const FPolygonID PolygonID, const in
 	int32 Index = 0;
 	for( const FVertexInstanceID VertexInstanceID : Contour.VertexInstanceIDs )
 	{
-		const FMeshVertexInstance& VertexInstance = VertexInstances[ VertexInstanceID.GetValue() ];
+		const FMeshVertexInstance& VertexInstance = VertexInstances[ VertexInstanceID ];
 		OutHoleVertexIDs[ Index ] = VertexInstance.VertexID;
 		Index++;
 	}
@@ -1748,8 +1593,7 @@ void UEditableMesh::GetPolygonHoleVertices( const FPolygonID PolygonID, const in
 
 void UEditableMesh::GetPolygonHoleVertexInstances( const FPolygonID PolygonID, const int32 HoleNumber, TArray<FVertexInstanceID>& OutHoleVertexInstanceIDs ) const
 {
-	checkSlow( Polygons.IsAllocated( PolygonID.GetValue() ) );
-	const FMeshPolygon& Polygon = Polygons[ PolygonID.GetValue() ];
+	const FMeshPolygon& Polygon = GetMeshDescription()->GetPolygon( PolygonID );
 	const FMeshPolygonContour& Contour = Polygon.HoleContours[ HoleNumber ];
 
 	OutHoleVertexInstanceIDs = Contour.VertexInstanceIDs;
@@ -1763,10 +1607,9 @@ FEdgeID UEditableMesh::GetPolygonContourEdge( const FMeshPolygonContour& Contour
 	const FVertexInstanceID VertexInstanceID = Contour.VertexInstanceIDs[ ContourEdgeNumber ];
 	const FVertexInstanceID NextVertexInstanceID = Contour.VertexInstanceIDs[ ( ContourEdgeNumber + 1 ) % NumPolygonContourEdges ];
 
-	checkSlow( VertexInstances.IsAllocated( VertexInstanceID.GetValue() ) );
-	checkSlow( VertexInstances.IsAllocated( NextVertexInstanceID.GetValue() ) );
-	const FVertexID VertexID = VertexInstances[ VertexInstanceID.GetValue() ].VertexID;
-	const FVertexID NextVertexID = VertexInstances[ NextVertexInstanceID.GetValue() ].VertexID;
+	const FVertexInstanceArray& VertexInstances = GetMeshDescription()->VertexInstances();
+	const FVertexID VertexID = VertexInstances[ VertexInstanceID ].VertexID;
+	const FVertexID NextVertexID = VertexInstances[ NextVertexInstanceID ].VertexID;
 
 	return GetVertexPairEdge( VertexID, NextVertexID, bOutEdgeWindingIsReversed );
 }
@@ -1774,9 +1617,7 @@ FEdgeID UEditableMesh::GetPolygonContourEdge( const FMeshPolygonContour& Contour
 
 FEdgeID UEditableMesh::GetPolygonPerimeterEdge( const FPolygonID PolygonID, const int32 PerimeterEdgeNumber, bool& bOutEdgeWindingIsReversedForPolygon ) const
 {
-	checkSlow( Polygons.IsAllocated( PolygonID.GetValue() ) );
-	const FMeshPolygon& Polygon = Polygons[ PolygonID.GetValue() ];
-	return GetPolygonContourEdge( Polygon.PerimeterContour, PerimeterEdgeNumber, bOutEdgeWindingIsReversedForPolygon );
+	return GetPolygonContourEdge( GetMeshDescription()->GetPolygon( PolygonID ).PerimeterContour, PerimeterEdgeNumber, bOutEdgeWindingIsReversedForPolygon );
 }
 
 
@@ -1784,10 +1625,11 @@ FEdgeID UEditableMesh::GetVertexPairEdge( const FVertexID StartVertexID, const F
 {
 	bOutEdgeWindingIsReversed = false;
 
-	const FMeshVertex& StartVertex = Vertices[ StartVertexID.GetValue() ];
-	for( const FEdgeID VertexConnectedEdgeID : StartVertex.ConnectedEdgeIDs )
+	const FEdgeArray& Edges = GetMeshDescription()->Edges();
+
+	for( const FEdgeID VertexConnectedEdgeID : GetMeshDescription()->GetVertex( StartVertexID ).ConnectedEdgeIDs )
 	{
-		const FMeshEdge& Edge = Edges[ VertexConnectedEdgeID.GetValue() ];
+		const FMeshEdge& Edge = Edges[ VertexConnectedEdgeID ];
 
 		// Try the edge's first vertex.  Does it point to us?
 		FVertexID OtherEdgeVertexID = Edge.VertexIDs[ 0 ];
@@ -1812,12 +1654,8 @@ FEdgeID UEditableMesh::GetVertexPairEdge( const FVertexID StartVertexID, const F
 
 FEdgeID UEditableMesh::GetPolygonHoleEdge( const FPolygonID PolygonID, const int32 HoleNumber, const int32 HoleEdgeNumber ) const
 {
-	checkSlow( Polygons.IsAllocated( PolygonID.GetValue() ) );
-	const FMeshPolygon& Polygon = Polygons[ PolygonID.GetValue() ];
-	const FMeshPolygonContour& Contour = Polygon.HoleContours[ HoleNumber ];
-
 	bool bOutEdgeWindingIsReversed;
-	return GetPolygonContourEdge( Contour, HoleEdgeNumber, bOutEdgeWindingIsReversed );
+	return GetPolygonContourEdge( GetMeshDescription()->GetPolygon( PolygonID ).HoleContours[ HoleEdgeNumber ], HoleEdgeNumber, bOutEdgeWindingIsReversed );
 }
 
 
@@ -1836,57 +1674,18 @@ void UEditableMesh::GetPolygonContourEdges( const FMeshPolygonContour& Contour,
 
 void UEditableMesh::GetPolygonPerimeterEdges( const FPolygonID PolygonID, TArray<FEdgeID>& OutPolygonPerimeterEdgeIDs ) const
 {
-	checkSlow( Polygons.IsAllocated( PolygonID.GetValue() ) );
-	const FMeshPolygon& Polygon = Polygons[ PolygonID.GetValue() ];
-	GetPolygonContourEdges( Polygon.PerimeterContour, OutPolygonPerimeterEdgeIDs );
+	GetPolygonContourEdges( GetMeshDescription()->GetPolygon( PolygonID ).PerimeterContour, OutPolygonPerimeterEdgeIDs );
 }
 
 
 void UEditableMesh::GetPolygonHoleEdges( const FPolygonID PolygonID, const int32 HoleNumber, TArray<FEdgeID>& OutHoleEdgeIDs ) const
 {
-	checkSlow( Polygons.IsAllocated( PolygonID.GetValue() ) );
-	const FMeshPolygon& Polygon = Polygons[ PolygonID.GetValue() ];
+	const FMeshPolygon& Polygon = GetMeshDescription()->GetPolygon( PolygonID );
 	const FMeshPolygonContour& Contour = Polygon.HoleContours[ HoleNumber ];
 	GetPolygonContourEdges( Contour, OutHoleEdgeIDs );
 }
 
 
-int32 UEditableMesh::GetPolygonTriangulatedTriangleCount( const FPolygonID PolygonID ) const
-{
-	checkSlow( Polygons.IsAllocated( PolygonID.GetValue() ) );
-	const FMeshPolygon& Polygon = Polygons[ PolygonID.GetValue() ];
-
-	return Polygon.Triangles.Num();
-}
-
-
-FMeshTriangle UEditableMesh::GetPolygonTriangulatedTriangle( const FPolygonID PolygonID, int32 PolygonTriangleNumber ) const
-{
-	checkSlow( Polygons.IsAllocated( PolygonID.GetValue() ) );
-	const FMeshPolygon& Polygon = Polygons[ PolygonID.GetValue() ];
-
-	// @todo mesheditor: how to make this safe so that BPs don't crash when given an out-of-range index.
-	// (same goes for just about all of these methods)
-	return Polygon.Triangles[ PolygonTriangleNumber ];
-}
-
-
-FVector UEditableMesh::GetPolygonTriangulatedTriangleVertexPosition( const FPolygonID PolygonID, const int32 PolygonTriangleNumber, const int32 TriangleVertexNumber ) const
-{
-	checkSlow( TriangleVertexNumber >= 0 && TriangleVertexNumber < 3 );
-
-	checkSlow( Polygons.IsAllocated( PolygonID.GetValue() ) );
-	const FMeshPolygon& Polygon = Polygons[ PolygonID.GetValue() ];
-
-	const FVertexInstanceID VertexInstanceID = Polygon.Triangles[ PolygonTriangleNumber ].GetVertexInstanceID( TriangleVertexNumber );
-	checkSlow( VertexInstances.IsAllocated( VertexInstanceID.GetValue() ) );
-	const FVertexID VertexID = VertexInstances[ VertexInstanceID.GetValue() ].VertexID;
-
-	checkSlow( Vertices.IsAllocated( VertexID.GetValue() ) );
-	return Vertices[ VertexID.GetValue() ].VertexPosition;
-}
-
-
 void UEditableMesh::GetPolygonAdjacentPolygons( const FPolygonID PolygonID, TArray<FPolygonID>& OutAdjacentPolygons ) const
 {
 	OutAdjacentPolygons.Reset();
@@ -1915,9 +1714,9 @@ FBox UEditableMesh::ComputeBoundingBox() const
 	FBox BoundingBox;
 	BoundingBox.Init();
 
-	for( TSparseArray<FMeshVertex>::TConstIterator It( Vertices ); It; ++It )
+	for( const FVertexID VertexID : GetMeshDescription()->Vertices().GetElementIDs() )
 	{
-		BoundingBox += It->VertexPosition;
+		BoundingBox += GetVertexAttribute( VertexID, UEditableMeshAttribute::VertexPosition(), 0 );
 	}
 
 	return BoundingBox;
@@ -1934,16 +1733,10 @@ FBoxSphereBounds UEditableMesh::ComputeBoundingBoxAndSphere() const
 	// Calculate the bounding sphere, using the center of the bounding box as the origin.
 	BoundingBoxAndSphere.SphereRadius = 0.0f;
 
-	const int32 VertexArraySize = GetVertexArraySize();
-	for( int32 VertexNumber = 0; VertexNumber < VertexArraySize; ++VertexNumber )
+	for( const FVertexID VertexID : GetMeshDescription()->Vertices().GetElementIDs() )
 	{
-		const FVertexID VertexID( VertexNumber );
-		if( IsValidVertex( VertexID ) )
-		{
-			const FVector VertexPosition = GetVertexAttribute( VertexID, UEditableMeshAttribute::VertexPosition(), 0 );
-
-			BoundingBoxAndSphere.SphereRadius = FMath::Max( ( VertexPosition - BoundingBoxAndSphere.Origin ).Size(), BoundingBoxAndSphere.SphereRadius );
-		}
+		const FVector VertexPosition = GetVertexAttribute( VertexID, UEditableMeshAttribute::VertexPosition(), 0 );
+		BoundingBoxAndSphere.SphereRadius = FMath::Max( ( VertexPosition - BoundingBoxAndSphere.Origin ).Size(), BoundingBoxAndSphere.SphereRadius );
 	}
 
 	return BoundingBoxAndSphere;
@@ -1960,7 +1753,6 @@ FVector UEditableMesh::ComputePolygonCenter( const FPolygonID PolygonID ) const
 	for( const FVertexID VertexID : PerimeterVertexIDs )
 	{
 		const FVector Position = GetVertexAttribute( VertexID, UEditableMeshAttribute::VertexPosition(), 0 );
-
 		Centroid += Position;
 	}
 
@@ -2029,7 +1821,7 @@ void UEditableMesh::RefreshOpenSubdiv()
 		// for our mesh (http://graphics.pixar.com/opensubdiv/docs/far_overview.html)
 		OpenSubdiv::Far::TopologyDescriptor OsdTopologyDescriptor;
 		{
-			const int32 VertexArraySize = GetVertexArraySize();
+			const int32 VertexArraySize = GetMeshDescription()->Vertices().GetArraySize();
 			OsdTopologyDescriptor.numVertices = VertexArraySize;
 			OsdTopologyDescriptor.numFaces = GetPolygonCount();
 
@@ -2041,44 +1833,34 @@ void UEditableMesh::RefreshOpenSubdiv()
 				{
 					OsdCornerVertexIndices.Reset();
 					OsdCornerWeights.Reset();
-					for( int32 VertexNumber = 0; VertexNumber < VertexArraySize; ++VertexNumber )
+					for( const FVertexID VertexID : GetMeshDescription()->Vertices().GetElementIDs() )
 					{
-						const FVertexID VertexID( VertexNumber );
-						if( IsValidVertex( VertexID ) )
+						const float VertexCornerSharpness = GetVertexAttribute( VertexID, UEditableMeshAttribute::VertexCornerSharpness(), 0 ).X;
+						if( VertexCornerSharpness > SMALL_NUMBER )
 						{
-							const float VertexCornerSharpness = GetVertexAttribute( VertexID, UEditableMeshAttribute::VertexCornerSharpness(), 0 ).X;
-							if( VertexCornerSharpness > SMALL_NUMBER )
-							{
-								// This vertex is (at least partially) a subdivision corner
-								OsdCornerVertexIndices.Add( VertexNumber );
-								OsdCornerWeights.Add( OpenSubdivCreaseWeightMultiplier * VertexCornerSharpness );
-							}
+							// This vertex is (at least partially) a subdivision corner
+							OsdCornerVertexIndices.Add( VertexID.GetValue() );
+							OsdCornerWeights.Add( OpenSubdivCreaseWeightMultiplier * VertexCornerSharpness );
 						}
 					}
 				}
 
 				// Edge creases
 				{
-					const int32 EdgeArraySize = GetEdgeArraySize();
-
 					OsdCreaseVertexIndexPairs.Reset();
 					OsdCreaseWeights.Reset();
-					for( int32 EdgeNumber = 0; EdgeNumber < EdgeArraySize; ++EdgeNumber )
+					for( const FEdgeID EdgeID : GetMeshDescription()->Edges().GetElementIDs() )
 					{
-						const FEdgeID EdgeID( EdgeNumber );
-						if( IsValidEdge( EdgeID ) )
+						const float EdgeCreaseSharpness = GetEdgeAttribute( EdgeID, UEditableMeshAttribute::EdgeCreaseSharpness(), 0 ).X;
+						if( EdgeCreaseSharpness > SMALL_NUMBER )
 						{
-							const float EdgeCreaseSharpness = GetEdgeAttribute( EdgeID, UEditableMeshAttribute::EdgeCreaseSharpness(), 0 ).X;
-							if( EdgeCreaseSharpness > SMALL_NUMBER )
-							{
-								// This edge is (at least partially) creased
-								FVertexID EdgeVertexID0, EdgeVertexID1;
-								GetEdgeVertices( EdgeID, /* Out */ EdgeVertexID0, /* Out */ EdgeVertexID1 );
+							// This edge is (at least partially) creased
+							FVertexID EdgeVertexID0, EdgeVertexID1;
+							GetEdgeVertices( EdgeID, /* Out */ EdgeVertexID0, /* Out */ EdgeVertexID1 );
 
-								OsdCreaseVertexIndexPairs.Add( EdgeVertexID0.GetValue() );
-								OsdCreaseVertexIndexPairs.Add( EdgeVertexID1.GetValue() );
-								OsdCreaseWeights.Add( OpenSubdivCreaseWeightMultiplier * EdgeCreaseSharpness );
-							}
+							OsdCreaseVertexIndexPairs.Add( EdgeVertexID0.GetValue() );
+							OsdCreaseVertexIndexPairs.Add( EdgeVertexID1.GetValue() );
+							OsdCreaseWeights.Add( OpenSubdivCreaseWeightMultiplier * EdgeCreaseSharpness );
 						}
 					}
 				}
@@ -2089,25 +1871,20 @@ void UEditableMesh::RefreshOpenSubdiv()
 
 				int32 NextOsdFaceIndex = 0;
 
-				const int32 PolygonCount = GetPolygonArraySize();
-				for( int32 PolygonIndex = 0; PolygonIndex < PolygonCount; ++PolygonIndex )
+				for( const FPolygonID PolygonID : GetMeshDescription()->Polygons().GetElementIDs() )
 				{
-					const FPolygonID PolygonID( PolygonIndex );
-					if( IsValidPolygon( PolygonID ) )
-					{
-						static TArray<FVertexID> PerimeterVertexIDs;
-						GetPolygonPerimeterVertices( PolygonID, /* Out */ PerimeterVertexIDs );
+					static TArray<FVertexID> PerimeterVertexIDs;
+					GetPolygonPerimeterVertices( PolygonID, /* Out */ PerimeterVertexIDs );
 
-						const int32 PerimeterVertexCount = PerimeterVertexIDs.Num();
-						OsdNumVerticesPerFace[ NextOsdFaceIndex++ ] = PerimeterVertexCount;
+					const int32 PerimeterVertexCount = PerimeterVertexIDs.Num();
+					OsdNumVerticesPerFace[ NextOsdFaceIndex++ ] = PerimeterVertexCount;
 
-						for( int32 PerimeterVertexNumber = 0; PerimeterVertexNumber < PerimeterVertexCount; ++PerimeterVertexNumber )
-						{
-							const FVertexID PerimeterVertexID = PerimeterVertexIDs[ PerimeterVertexNumber ];
+					for( int32 PerimeterVertexNumber = 0; PerimeterVertexNumber < PerimeterVertexCount; ++PerimeterVertexNumber )
+					{
+						const FVertexID PerimeterVertexID = PerimeterVertexIDs[ PerimeterVertexNumber ];
 
-							OsdVertexIndicesPerFace.Add( PerimeterVertexID.GetValue() );
-							OsdFVarIndicesPerFace.Add( OsdFVarIndicesPerFace.Num() );
-						}
+						OsdVertexIndicesPerFace.Add( PerimeterVertexID.GetValue() );
+						OsdFVarIndicesPerFace.Add( OsdFVarIndicesPerFace.Num() );
 					}
 				}
 
@@ -2259,7 +2036,7 @@ void UEditableMesh::GenerateOpenSubdivLimitSurfaceData()
 			{
 				// NOTE: We're including an entry for all vertices, even vertices that aren't referenced by any triangles (due to our
 				//       sparse array optimization.)
-				const int32 VertexArraySize = GetVertexArraySize();
+				const int32 VertexArraySize = GetMeshDescription()->Vertices().GetArraySize();
 				VertexPositions.SetNum( VertexArraySize, false );
 				for( int32 VertexNumber = 0; VertexNumber < VertexArraySize; ++VertexNumber )
 				{
@@ -2283,30 +2060,24 @@ void UEditableMesh::GenerateOpenSubdivLimitSurfaceData()
 				FVarVertexDatas.Reset();
 				FVarVertexDatas.Reserve( OsdFVarIndicesPerFace.Num() );
 
-				const int32 PolygonGroupArraySize = GetPolygonGroupArraySize();
-
 				int32 NumPolygonsSoFar = 0;
-				for( int32 PolygonGroupNumber = 0; PolygonGroupNumber < PolygonGroupArraySize; ++PolygonGroupNumber )
+				for( const FPolygonGroupID PolygonGroupID : GetMeshDescription()->PolygonGroups().GetElementIDs() )
 				{
-					const FPolygonGroupID PolygonGroupID( PolygonGroupNumber );
-					if( IsValidPolygonGroup( PolygonGroupID ) )
+					FirstPolygonNumberForPolygonGroups.Add( NumPolygonsSoFar );
+					const int32 PolygonCountInGroup = GetPolygonCountInGroup( PolygonGroupID );
+					NumPolygonsSoFar += PolygonCountInGroup;
+
+					for( int32 PolygonGroupPolygonNumber = 0; PolygonGroupPolygonNumber < PolygonCountInGroup; ++PolygonGroupPolygonNumber )
 					{
-						FirstPolygonNumberForPolygonGroups.Add( NumPolygonsSoFar );
-						const int32 PolygonCountInGroup = GetPolygonCountInGroup( PolygonGroupID );
-						NumPolygonsSoFar += PolygonCountInGroup;
+						const FPolygonID PolygonID( GetPolygonInGroup( PolygonGroupID, PolygonGroupPolygonNumber ) );
 
-						for( int32 PolygonGroupPolygonNumber = 0; PolygonGroupPolygonNumber < PolygonCountInGroup; ++PolygonGroupPolygonNumber )
+						const int32 PerimeterVertexCount = GetPolygonPerimeterVertexCount( PolygonID );
+						for( int32 PerimeterVertexNumber = 0; PerimeterVertexNumber < PerimeterVertexCount; ++PerimeterVertexNumber )
 						{
-							const FPolygonID PolygonID( GetPolygonInGroup( PolygonGroupID, PolygonGroupPolygonNumber ) );
-
-							const int32 PerimeterVertexCount = GetPolygonPerimeterVertexCount( PolygonID );
-							for( int32 PerimeterVertexNumber = 0; PerimeterVertexNumber < PerimeterVertexCount; ++PerimeterVertexNumber )
-							{
-								FOsdFVarVertexData& FVarVertexData = *new( FVarVertexDatas ) FOsdFVarVertexData();
-								FVarVertexData.TextureCoordinates[ 0 ] = TextureCoordinateCount > 0 ? FVector2D( GetPolygonPerimeterVertexAttribute( PolygonID, PerimeterVertexNumber, UEditableMeshAttribute::VertexTextureCoordinate(), 0 ) ) : FVector2D::ZeroVector;
-								FVarVertexData.TextureCoordinates[ 1 ] = TextureCoordinateCount > 1 ? FVector2D( GetPolygonPerimeterVertexAttribute( PolygonID, PerimeterVertexNumber, UEditableMeshAttribute::VertexTextureCoordinate(), 1 ) ) : FVector2D::ZeroVector;
-								FVarVertexData.VertexColor = FLinearColor( GetPolygonPerimeterVertexAttribute( PolygonID, PerimeterVertexNumber, UEditableMeshAttribute::VertexColor(), 0 ) );
-							}
+							FOsdFVarVertexData& FVarVertexData = *new( FVarVertexDatas ) FOsdFVarVertexData();
+							FVarVertexData.TextureCoordinates[ 0 ] = TextureCoordinateCount > 0 ? FVector2D( GetPolygonPerimeterVertexAttribute( PolygonID, PerimeterVertexNumber, UEditableMeshAttribute::VertexTextureCoordinate(), 0 ) ) : FVector2D::ZeroVector;
+							FVarVertexData.TextureCoordinates[ 1 ] = TextureCoordinateCount > 1 ? FVector2D( GetPolygonPerimeterVertexAttribute( PolygonID, PerimeterVertexNumber, UEditableMeshAttribute::VertexTextureCoordinate(), 1 ) ) : FVector2D::ZeroVector;
+							FVarVertexData.VertexColor = FLinearColor( GetPolygonPerimeterVertexAttribute( PolygonID, PerimeterVertexNumber, UEditableMeshAttribute::VertexColor(), 0 ) );
 						}
 					}
 				}
@@ -2670,10 +2441,12 @@ void UEditableMesh::GenerateOpenSubdivLimitSurfaceData()
 
 void UEditableMesh::RetriangulatePolygons()
 {
+	FPolygonArray& Polygons = GetMeshDescription()->Polygons();
+
 	// Perform triangulation directly into mesh polygons
 	for( const FPolygonID PolygonID : PolygonsPendingTriangulation )
 	{
-		FMeshPolygon& Polygon = Polygons[ PolygonID.GetValue() ];
+		FMeshPolygon& Polygon = Polygons[ PolygonID ];
 		ComputePolygonTriangulation( PolygonID, Polygon.Triangles );
 	}
 
@@ -2708,7 +2481,7 @@ void UEditableMesh::ComputePolygonTriangulation( const FPolygonID PolygonID, TAr
 	// @todo mesheditor holes: Does not support triangles with holes yet!
 	// @todo mesheditor: Perhaps should always attempt to triangulate by splitting polygons along the shortest edge, for better determinism.
 
-	const FMeshPolygon& Polygon = Polygons[ PolygonID.GetValue() ];
+	const FMeshPolygon& Polygon = GetMeshDescription()->GetPolygon( PolygonID );
 	const TArray<FVertexInstanceID>& PolygonVertexInstanceIDs = Polygon.PerimeterContour.VertexInstanceIDs;
 
 	// Polygon must have at least three vertices/edges
@@ -2836,7 +2609,7 @@ bool UEditableMesh::ComputeBarycentricWeightForPointOnPolygon( const FPolygonID
 {
 	// @todo mesheditor: Modify this method so it can return a meaningful barycentric weight for an off-polygon position, referencing the closest triangle, with out-of-range weights.
 
-	const FMeshPolygon& Polygon = Polygons[ PolygonID.GetValue() ];
+	const FMeshPolygon& Polygon = GetMeshDescription()->GetPolygon( PolygonID );
 	const int32 TriangleCount = Polygon.Triangles.Num();
 
 	// Figure out which triangle the incoming point is within
@@ -3942,6 +3715,8 @@ void UEditableMesh::DeleteOrphanVertices( const TArray<FVertexID>& VertexIDsToDe
 {
 	UE_LOG( LogEditableMesh, Verbose, TEXT( "DeleteOrphanVertices: %s" ), *LogHelpers::ArrayToString( VertexIDsToDelete ) );
 
+	FVertexArray& Vertices = GetMeshDescription()->Vertices();
+
 	// Back everything up
 	{
 		FCreateVerticesChangeInput RevertInput;
@@ -3954,7 +3729,7 @@ void UEditableMesh::DeleteOrphanVertices( const TArray<FVertexID>& VertexIDsToDe
 
 			// Make sure the vertex is truly an orphan.  We're not going to be able to restore it's polygon vertex attributes,
 			// because the polygons won't exist when we're restoring the change
-			const FMeshVertex& Vertex = Vertices[ VertexID.GetValue() ];
+			const FMeshVertex& Vertex = Vertices[ VertexID ];
 			check( Vertex.ConnectedEdgeIDs.Num() == 0 );
 			check( Vertex.VertexInstanceIDs.Num() == 0 );
 
@@ -3990,7 +3765,7 @@ void UEditableMesh::DeleteOrphanVertices( const TArray<FVertexID>& VertexIDsToDe
 	{
 		for( const FVertexID VertexIDToDelete : VertexIDsToDelete )
 		{
-			Vertices.RemoveAt( VertexIDToDelete.GetValue() );
+			Vertices.Remove( VertexIDToDelete );
 		}
 	}
 
@@ -4002,6 +3777,9 @@ void UEditableMesh::DeleteVertexInstances( const TArray<FVertexInstanceID>& Vert
 {
 	UE_LOG( LogEditableMesh, Verbose, TEXT( "DeleteVertexInstances: %s" ), *LogHelpers::ArrayToString( VertexInstanceIDsToDelete ) );
 
+	FVertexInstanceArray& VertexInstances = GetMeshDescription()->VertexInstances();
+	FVertexArray& Vertices = GetMeshDescription()->Vertices();
+
 	// Back everything up
 	{
 		FCreateVertexInstancesChangeInput RevertInput;
@@ -4011,7 +3789,7 @@ void UEditableMesh::DeleteVertexInstances( const TArray<FVertexInstanceID>& Vert
 		for( int32 VertexNumber = VertexInstanceIDsToDelete.Num() - 1; VertexNumber >= 0; --VertexNumber )
 		{
 			const FVertexInstanceID VertexInstanceID = VertexInstanceIDsToDelete[ VertexNumber ];
-			const FMeshVertexInstance& VertexInstance = VertexInstances[ VertexInstanceID.GetValue() ];
+			const FMeshVertexInstance& VertexInstance = VertexInstances[ VertexInstanceID ];
 
 			// Back up properties
 			RevertInput.VertexInstancesToCreate.Emplace();
@@ -4054,8 +3832,8 @@ void UEditableMesh::DeleteVertexInstances( const TArray<FVertexInstanceID>& Vert
 			check( GetVertexInstanceConnectedPolygonCount( VertexInstanceIDToDelete ) == 0 );
 
 			// Remove the reference in the vertex it instances
-			const FMeshVertexInstance& VertexInstance = VertexInstances[ VertexInstanceIDToDelete.GetValue() ];
-			FMeshVertex& Vertex = Vertices[ VertexInstance.VertexID.GetValue() ];
+			const FMeshVertexInstance& VertexInstance = VertexInstances[ VertexInstanceIDToDelete ];
+			FMeshVertex& Vertex = Vertices[ VertexInstance.VertexID ];
 			verify( Vertex.VertexInstanceIDs.Remove( VertexInstanceIDToDelete ) == 1 );
 
 			if( bDeleteOrphanedVertices )
@@ -4067,7 +3845,7 @@ void UEditableMesh::DeleteVertexInstances( const TArray<FVertexInstanceID>& Vert
 			}
 
 			// Delete the element from the container
-			VertexInstances.RemoveAt( VertexInstanceIDToDelete.GetValue() );
+			VertexInstances.Remove( VertexInstanceIDToDelete );
 		}
 	}
 
@@ -4085,6 +3863,9 @@ void UEditableMesh::DeleteEdges( const TArray<FEdgeID>& EdgeIDsToDelete, const b
 {
 	UE_LOG( LogEditableMesh, Verbose, TEXT( "DeleteEdges: %s" ), *LogHelpers::ArrayToString( EdgeIDsToDelete ) );
 
+	FEdgeArray& Edges = GetMeshDescription()->Edges();
+	FVertexArray& Vertices = GetMeshDescription()->Vertices();
+
 	// Back everything up
 	{
 		FCreateEdgesChangeInput RevertInput;
@@ -4093,7 +3874,7 @@ void UEditableMesh::DeleteEdges( const TArray<FEdgeID>& EdgeIDsToDelete, const b
 		for( int32 EdgeNumber = EdgeIDsToDelete.Num() - 1; EdgeNumber >= 0; --EdgeNumber )
 		{
 			const FEdgeID EdgeID = EdgeIDsToDelete[ EdgeNumber ];
-			const FMeshEdge& Edge = Edges[ EdgeID.GetValue() ];
+			const FMeshEdge& Edge = Edges[ EdgeID ];
 
 			RevertInput.EdgesToCreate.Emplace();
 			FEdgeToCreate& EdgeToCreate = RevertInput.EdgesToCreate.Last();
@@ -4132,11 +3913,11 @@ void UEditableMesh::DeleteEdges( const TArray<FEdgeID>& EdgeIDsToDelete, const b
 
 		for( const FEdgeID EdgeID : EdgeIDsToDelete )
 		{
-			FMeshEdge& Edge = Edges[ EdgeID.GetValue() ];
+			FMeshEdge& Edge = Edges[ EdgeID ];
 
 			for( const FVertexID EdgeVertexID : Edge.VertexIDs )
 			{
-				FMeshVertex& Vertex = Vertices[ EdgeVertexID.GetValue() ];
+				FMeshVertex& Vertex = Vertices[ EdgeVertexID ];
 				verify( Vertex.ConnectedEdgeIDs.RemoveSingle( EdgeID ) == 1 );
 
 				// If the vertex has no more edges connected, we'll keep track of that so we can delete the vertex later
@@ -4149,7 +3930,7 @@ void UEditableMesh::DeleteEdges( const TArray<FEdgeID>& EdgeIDsToDelete, const b
 			}
 
 			// Delete the edge
-			Edges.RemoveAt( EdgeID.GetValue() );
+			Edges.Remove( EdgeID );
 		}
 
 		// If we orphaned any vertices and we were asked to delete those, then we'll go ahead and do that now.
@@ -4167,6 +3948,8 @@ void UEditableMesh::CreateEmptyVertexRange( const int32 NumVerticesToCreate, TAr
 {
 	UE_LOG( LogEditableMesh, Verbose, TEXT( "CreateEmptyVertexRange: %d" ), NumVerticesToCreate );
 
+	FVertexArray& Vertices = GetMeshDescription()->Vertices();
+
 	OutNewVertexIDs.Reset( NumVerticesToCreate );
 
 	// Reserve elements
@@ -4175,8 +3958,8 @@ void UEditableMesh::CreateEmptyVertexRange( const int32 NumVerticesToCreate, TAr
 	for( int32 Count = 0; Count < NumVerticesToCreate; ++Count )
 	{
 		// Allocate vertex
-		const FVertexID VertexID = FVertexID( Vertices.Add( FMeshVertex() ) );
-		FMeshVertex& Vertex = Vertices[ VertexID.GetValue() ];
+		const FVertexID VertexID = Vertices.Add();
+		FMeshVertex& Vertex = Vertices[ VertexID ];
 
 		// Initialize vertex
 		Vertex.VertexPosition = FVector::ZeroVector;
@@ -4211,6 +3994,8 @@ void UEditableMesh::CreateVertices( const TArray<FVertexToCreate>& VerticesToCre
 {
 	UE_LOG( LogEditableMesh, Verbose, TEXT( "CreateVertices: %s" ), *LogHelpers::ArrayToString( VerticesToCreate ) );
 
+	FVertexArray& Vertices = GetMeshDescription()->Vertices();
+
 	OutNewVertexIDs.Reset( VerticesToCreate.Num() );
 
 	// Reserve elements
@@ -4222,14 +4007,14 @@ void UEditableMesh::CreateVertices( const TArray<FVertexToCreate>& VerticesToCre
 		FVertexID VertexID = VertexToCreate.OriginalVertexID;
 		if( VertexID != FVertexID::Invalid )
 		{
-			Vertices.Insert( VertexID.GetValue(), FMeshVertex() );
+			Vertices.Insert( VertexID );
 		}
 		else
 		{
-			VertexID = FVertexID( Vertices.Add( FMeshVertex() ) );
+			VertexID = Vertices.Add();
 		}
 
-		FMeshVertex& Vertex = Vertices[ VertexID.GetValue() ];
+		FMeshVertex& Vertex = Vertices[ VertexID ];
 
 		// Initialize vertex
 		Vertex.VertexPosition = FVector::ZeroVector;
@@ -4273,6 +4058,9 @@ void UEditableMesh::CreateVertexInstances( const TArray<FVertexInstanceToCreate>
 {
 	UE_LOG( LogEditableMesh, Verbose, TEXT( "CreateVertexInstances: %s" ), *LogHelpers::ArrayToString( VertexInstancesToCreate ) );
 
+	FVertexArray& Vertices = GetMeshDescription()->Vertices();
+	FVertexInstanceArray& VertexInstances = GetMeshDescription()->VertexInstances();
+
 	OutNewVertexInstanceIDs.Reset( VertexInstancesToCreate.Num() );
 
 	// Reserve elements
@@ -4284,14 +4072,14 @@ void UEditableMesh::CreateVertexInstances( const TArray<FVertexInstanceToCreate>
 		FVertexInstanceID VertexInstanceID = VertexInstanceToCreate.OriginalVertexInstanceID;
 		if( VertexInstanceID != FVertexInstanceID::Invalid )
 		{
-			VertexInstances.Insert( VertexInstanceID.GetValue(), FMeshVertexInstance() );
+			VertexInstances.Insert( VertexInstanceID );
 		}
 		else
 		{
-			VertexInstanceID = FVertexInstanceID( VertexInstances.Add( FMeshVertexInstance() ) );
+			VertexInstanceID = VertexInstances.Add();
 		}
 
-		FMeshVertexInstance& VertexInstance = VertexInstances[ VertexInstanceID.GetValue() ];
+		FMeshVertexInstance& VertexInstance = VertexInstances[ VertexInstanceID ];
 
 		// Initialize vertex instance
 		VertexInstance.VertexID = VertexInstanceToCreate.VertexID;
@@ -4303,8 +4091,8 @@ void UEditableMesh::CreateVertexInstances( const TArray<FVertexInstanceToCreate>
 		VertexInstance.VertexUVs.SetNumZeroed( TextureCoordinateCount );
 
 		// Add to the list of instances in the vertex
-		check( !Vertices[ VertexInstance.VertexID.GetValue() ].VertexInstanceIDs.Contains( VertexInstanceID ) );
-		Vertices[ VertexInstance.VertexID.GetValue() ].VertexInstanceIDs.Add( VertexInstanceID );
+		check( !Vertices[ VertexInstance.VertexID ].VertexInstanceIDs.Contains( VertexInstanceID ) );
+		Vertices[ VertexInstance.VertexID ].VertexInstanceIDs.Add( VertexInstanceID );
 
 		OutNewVertexInstanceIDs.Add( VertexInstanceID );
 	}
@@ -4345,6 +4133,9 @@ void UEditableMesh::CreateEdges( const TArray<FEdgeToCreate>& EdgesToCreate, TAr
 {
 	UE_LOG( LogEditableMesh, Verbose, TEXT( "CreateEdges: %s" ), *LogHelpers::ArrayToString( EdgesToCreate ) );
 
+	FVertexArray& Vertices = GetMeshDescription()->Vertices();
+	FEdgeArray& Edges = GetMeshDescription()->Edges();
+
 	OutNewEdgeIDs.Reset( EdgesToCreate.Num() );
 
 	// Reserve elements
@@ -4356,14 +4147,14 @@ void UEditableMesh::CreateEdges( const TArray<FEdgeToCreate>& EdgesToCreate, TAr
 		FEdgeID EdgeID = EdgeToCreate.OriginalEdgeID;
 		if( EdgeID != FEdgeID::Invalid )
 		{
-			Edges.Insert( EdgeID.GetValue(), FMeshEdge() );
+			Edges.Insert( EdgeID );
 		}
 		else
 		{
-			EdgeID = FEdgeID( Edges.Add( FMeshEdge() ) );
+			EdgeID = Edges.Add();
 		}
 
-		FMeshEdge& Edge = Edges[ EdgeID.GetValue() ];
+		FMeshEdge& Edge = Edges[ EdgeID ];
 
 		// Initialize edge
 		Edge.VertexIDs[ 0 ] = EdgeToCreate.VertexID0;
@@ -4374,8 +4165,8 @@ void UEditableMesh::CreateEdges( const TArray<FEdgeToCreate>& EdgesToCreate, TAr
 		Edge.CreaseSharpness = 0.0f;
 
 		// Connect the edge to its vertices
-		Vertices[ Edge.VertexIDs[ 0 ].GetValue() ].ConnectedEdgeIDs.AddUnique( EdgeID );
-		Vertices[ Edge.VertexIDs[ 1 ].GetValue() ].ConnectedEdgeIDs.AddUnique( EdgeID );
+		Vertices[ Edge.VertexIDs[ 0 ] ].ConnectedEdgeIDs.AddUnique( EdgeID );
+		Vertices[ Edge.VertexIDs[ 1 ] ].ConnectedEdgeIDs.AddUnique( EdgeID );
 
 		OutNewEdgeIDs.Add( EdgeID );
 	}
@@ -4482,6 +4273,9 @@ FVertexInstanceID UEditableMesh::CreateVertexInstanceForContourVertex( const FVe
 
 void UEditableMesh::CreatePolygonContour( const FPolygonID PolygonID, const TArray<FVertexAndAttributes>& Contour, const EPolygonEdgeHardness PolygonEdgeHardness, TArray<FEdgeID>& OutEdgeIDs, TArray<FVertexInstanceID>& OutVertexInstanceIDs )
 {
+	FVertexInstanceArray& VertexInstances = GetMeshDescription()->VertexInstances();
+	FEdgeArray& Edges = GetMeshDescription()->Edges();
+
 	// All polygons must have at least three vertices
 	const int32 NumContourVertices = Contour.Num();
 	check( NumContourVertices >= 3 );
@@ -4506,14 +4300,14 @@ void UEditableMesh::CreatePolygonContour( const FPolygonID PolygonID, const TArr
 		{
 			const FVertexInstanceID VertexInstanceID0 = Contour[ VertexNumber ].VertexInstanceID;
 			check( VertexInstanceID0 != FVertexInstanceID::Invalid );
-			VertexID0 = VertexInstances[ VertexInstanceID0.GetValue() ].VertexID;
+			VertexID0 = VertexInstances[ VertexInstanceID0 ].VertexID;
 		}
 
 		if( VertexID1 == FVertexID::Invalid )
 		{
 			const FVertexInstanceID VertexInstanceID1 = Contour[ NextVertexNumber ].VertexInstanceID;
 			check( VertexInstanceID1 != FVertexInstanceID::Invalid );
-			VertexID1 = VertexInstances[ VertexInstanceID1.GetValue() ].VertexID;
+			VertexID1 = VertexInstances[ VertexInstanceID1 ].VertexID;
 		}
 
 		bool bOutEdgeWindingIsReversed;
@@ -4543,7 +4337,7 @@ void UEditableMesh::CreatePolygonContour( const FPolygonID PolygonID, const TArr
 		}
 
 		// Connect the edge to the polygon
-		FMeshEdge& Edge = Edges[ EdgeID.GetValue() ];
+		FMeshEdge& Edge = Edges[ EdgeID ];
 		Edge.ConnectedPolygons.Add( PolygonID );
 
 		AllEdges[ VertexNumber ] = EdgeID;
@@ -4576,7 +4370,7 @@ void UEditableMesh::CreatePolygonContour( const FPolygonID PolygonID, const TArr
 		OutVertexInstanceIDs[ VertexNumber ] = ContourVertexInstanceID;
 
 		// Connect this vertex instance to the polygon
-		FMeshVertexInstance& ContourVertexInstance = VertexInstances[ ContourVertexInstanceID.GetValue() ];
+		FMeshVertexInstance& ContourVertexInstance = VertexInstances[ ContourVertexInstanceID ];
 		check( !ContourVertexInstance.ConnectedPolygons.Contains( PolygonID ) );
 		ContourVertexInstance.ConnectedPolygons.Add( PolygonID );
 	}
@@ -4587,6 +4381,9 @@ void UEditableMesh::CreatePolygons( const TArray<FPolygonToCreate>& PolygonsToCr
 {
 	UE_LOG( LogEditableMesh, Verbose, TEXT( "CreatePolygons: %s" ), *LogHelpers::ArrayToString( PolygonsToCreate ) );
 
+	FPolygonArray& Polygons = GetMeshDescription()->Polygons();
+	FPolygonGroupArray& PolygonGroups = GetMeshDescription()->PolygonGroups();
+
 	OutNewPolygonIDs.Reset( PolygonsToCreate.Num() );
 	OutNewEdgeIDs.Reset();
 
@@ -4598,16 +4395,16 @@ void UEditableMesh::CreatePolygons( const TArray<FPolygonToCreate>& PolygonsToCr
 		FPolygonID PolygonID = PolygonToCreate.OriginalPolygonID;
 		if( PolygonID != FPolygonID::Invalid )
 		{
-			Polygons.Insert( PolygonID.GetValue(), FMeshPolygon() );
+			Polygons.Insert( PolygonID );
 		}
 		else
 		{
-			PolygonID = FPolygonID( Polygons.Add( FMeshPolygon() ) );
+			PolygonID = Polygons.Add();
 		}
 
 		// Connect polygon to group and vice versa
-		FMeshPolygon& Polygon = Polygons[ PolygonID.GetValue() ];
-		FMeshPolygonGroup& PolygonGroup = PolygonGroups[ PolygonToCreate.PolygonGroupID.GetValue() ];
+		FMeshPolygon& Polygon = Polygons[ PolygonID ];
+		FMeshPolygonGroup& PolygonGroup = PolygonGroups[ PolygonToCreate.PolygonGroupID ];
 		Polygon.PolygonGroupID = PolygonToCreate.PolygonGroupID;
 		PolygonGroup.Polygons.Add( PolygonID );
 
@@ -4669,10 +4466,12 @@ void UEditableMesh::CreatePolygons( const TArray<FPolygonToCreate>& PolygonsToCr
 
 void UEditableMesh::BackupPolygonContour( const FMeshPolygonContour& Contour, TArray<FVertexAndAttributes>& OutVerticesAndAttributes )
 {
+	const FVertexInstanceArray& VertexInstances = GetMeshDescription()->VertexInstances();
+
 	OutVerticesAndAttributes.Reserve( Contour.VertexInstanceIDs.Num() );
 	for( const FVertexInstanceID VertexInstanceID : Contour.VertexInstanceIDs )
 	{
-		const FMeshVertexInstance& VertexInstance = VertexInstances[ VertexInstanceID.GetValue() ];
+		const FMeshVertexInstance& VertexInstance = VertexInstances[ VertexInstanceID ];
 
 		OutVerticesAndAttributes.Emplace();
 		FVertexAndAttributes& VertexAndAttributes = OutVerticesAndAttributes.Last();
@@ -4686,10 +4485,12 @@ void UEditableMesh::BackupPolygonContour( const FMeshPolygonContour& Contour, TA
 
 void UEditableMesh::DeletePolygonContour( const FPolygonID PolygonID, const FMeshPolygonContour& Contour, const bool bDeleteOrphanedEdges, const bool bDeleteOrphanedVertexInstances, TArray<FEdgeID>& OrphanedEdgeIDs, TArray<FVertexInstanceID>& OrphanedVertexInstanceIDs )
 {
+	FVertexInstanceArray& VertexInstances = GetMeshDescription()->VertexInstances();
+
 	for( const FVertexInstanceID VertexInstanceID : Contour.VertexInstanceIDs )
 	{
 		// Disconnect the polygon from the vertex instance
-		FMeshVertexInstance& VertexInstance = VertexInstances[ VertexInstanceID.GetValue() ];
+		FMeshVertexInstance& VertexInstance = VertexInstances[ VertexInstanceID ];
 		verify( VertexInstance.ConnectedPolygons.Remove( PolygonID ) == 1 );
 
 		// Vertex instances which are only part of this contour (i.e. not shared by other polygons) are optionally deleted
@@ -4703,9 +4504,11 @@ void UEditableMesh::DeletePolygonContour( const FPolygonID PolygonID, const FMes
 	static TArray<FEdgeID> EdgeIDs;
 	GetPolygonContourEdges( Contour, EdgeIDs );
 
+	FEdgeArray& Edges = GetMeshDescription()->Edges();
+
 	for( const FEdgeID EdgeID : EdgeIDs )
 	{
-		FMeshEdge& Edge = Edges[ EdgeID.GetValue() ];
+		FMeshEdge& Edge = Edges[ EdgeID ];
 		verify( Edge.ConnectedPolygons.Remove( PolygonID ) == 1 );
 
 		// Identify orphaned edges
@@ -4721,6 +4524,9 @@ void UEditableMesh::DeletePolygons( const TArray<FPolygonID>& PolygonIDsToDelete
 {
 	UE_LOG( LogEditableMesh, Verbose, TEXT( "DeletePolygons: %s" ), *LogHelpers::ArrayToString( PolygonIDsToDelete ) );
 
+	FPolygonArray& Polygons = GetMeshDescription()->Polygons();
+	FPolygonGroupArray& PolygonGroups = GetMeshDescription()->PolygonGroups();
+
 	// Back everything up
 	{
 		FCreatePolygonsChangeInput RevertInput;
@@ -4730,7 +4536,7 @@ void UEditableMesh::DeletePolygons( const TArray<FPolygonID>& PolygonIDsToDelete
 		for( int32 PolygonNumber = PolygonIDsToDelete.Num() - 1; PolygonNumber >= 0; --PolygonNumber )
 		{
 			const FPolygonID PolygonID = PolygonIDsToDelete[ PolygonNumber ];
-			const FMeshPolygon& Polygon = Polygons[ PolygonID.GetValue() ];
+			const FMeshPolygon& Polygon = Polygons[ PolygonID ];
 
 			RevertInput.PolygonsToCreate.Emplace();
 			FPolygonToCreate& PolygonToCreate = RevertInput.PolygonsToCreate.Last();
@@ -4769,7 +4575,7 @@ void UEditableMesh::DeletePolygons( const TArray<FPolygonID>& PolygonIDsToDelete
 
 		for( const FPolygonID PolygonID : PolygonIDsToDelete )
 		{
-			const FMeshPolygon& Polygon = Polygons[ PolygonID.GetValue() ];
+			const FMeshPolygon& Polygon = Polygons[ PolygonID ];
 
 			DeletePolygonContour( PolygonID, Polygon.PerimeterContour, bDeleteOrphanedEdges, bDeleteOrphanedVertexInstances, OrphanedEdgeIDs, OrphanedVertexInstanceIDs );
 
@@ -4780,7 +4586,7 @@ void UEditableMesh::DeletePolygons( const TArray<FPolygonID>& PolygonIDsToDelete
 
 			// Disconnect the polygon from the polygon group
 			const FPolygonGroupID PolygonGroupID = Polygon.PolygonGroupID;
-			FMeshPolygonGroup& PolygonGroup = PolygonGroups[ PolygonGroupID.GetValue() ];
+			FMeshPolygonGroup& PolygonGroup = PolygonGroups[ PolygonGroupID ];
 			verify( PolygonGroup.Polygons.Remove( PolygonID ) == 1 );
 
 			if( bDeleteEmptyPolygonGroups && PolygonGroup.Polygons.Num() == 0 )
@@ -4789,7 +4595,7 @@ void UEditableMesh::DeletePolygons( const TArray<FPolygonID>& PolygonIDsToDelete
 			}
 
 			// Delete the polygon
-			Polygons.RemoveAt( PolygonID.GetValue() );
+			Polygons.Remove( PolygonID );
 
 			// Update spatial database
 			if( Octree.IsValid() )
@@ -4841,6 +4647,10 @@ void UEditableMesh::DeletePolygons( const TArray<FPolygonID>& PolygonIDsToDelete
 
 void UEditableMesh::CreatePolygonGroups( const TArray<FPolygonGroupToCreate>& PolygonGroupsToCreate, TArray<FPolygonGroupID>& OutNewPolygonGroupIDs )
 {
+	UE_LOG( LogEditableMesh, Verbose, TEXT( "CreatePolygonGroups: %s" ), *LogHelpers::ArrayToString( PolygonGroupsToCreate ) );
+
+	FPolygonGroupArray& PolygonGroups = GetMeshDescription()->PolygonGroups();
+
 	// @todo mesheditor: Material, bEnableCollision and bCastShadow should all be implemented as PolygonGroup attributes
 
 	OutNewPolygonGroupIDs.Reset( PolygonGroupsToCreate.Num() );
@@ -4854,16 +4664,18 @@ void UEditableMesh::CreatePolygonGroups( const TArray<FPolygonGroupToCreate>& Po
 		FPolygonGroupID PolygonGroupID = PolygonGroupToCreate.OriginalPolygonGroupID;
 		if( PolygonGroupID != FPolygonGroupID::Invalid )
 		{
-			PolygonGroups.Insert( PolygonGroupID.GetValue(), FMeshPolygonGroup() );
+			PolygonGroups.Insert( PolygonGroupID );
 		}
 		else
 		{
-			PolygonGroupID = FPolygonGroupID( PolygonGroups.Add( FMeshPolygonGroup() ) );
+			PolygonGroupID = PolygonGroups.Add();
 		}
 
-		FMeshPolygonGroup& PolygonGroup = PolygonGroups[ PolygonGroupID.GetValue() ];
+		FMeshPolygonGroup& PolygonGroup = PolygonGroups[ PolygonGroupID ];
 
-		PolygonGroup.Material = PolygonGroupToCreate.Material;
+		PolygonGroup.MaterialAsset = PolygonGroupToCreate.MaterialAsset;
+		PolygonGroup.MaterialSlotName = PolygonGroupToCreate.MaterialSlotName;
+		PolygonGroup.ImportedMaterialSlotName = PolygonGroupToCreate.ImportedMaterialSlotName;
 		PolygonGroup.bEnableCollision = PolygonGroupToCreate.bEnableCollision;
 		PolygonGroup.bCastShadow = PolygonGroupToCreate.bCastShadow;
 
@@ -4886,11 +4698,17 @@ void UEditableMesh::CreatePolygonGroups( const TArray<FPolygonGroupToCreate>& Po
 
 		AddUndo( MakeUnique<FDeletePolygonGroupsChange>( MoveTemp( RevertInput ) ) );
 	}
+
+	UE_LOG( LogEditableMesh, Verbose, TEXT( "* CreatePolygonGroups returned %s" ), *LogHelpers::ArrayToString( OutNewPolygonGroupIDs ) );
 }
 
 
 void UEditableMesh::DeletePolygonGroups( const TArray<FPolygonGroupID>& PolygonGroupIDs )
 {
+	UE_LOG( LogEditableMesh, Verbose, TEXT( "DeletePolygonGroups: %s" ), *LogHelpers::ArrayToString( PolygonGroupIDs ) );
+
+	FPolygonGroupArray& PolygonGroups = GetMeshDescription()->PolygonGroups();
+
 	// Back everything up
 	{
 		FCreatePolygonGroupsChangeInput RevertInput;
@@ -4898,13 +4716,15 @@ void UEditableMesh::DeletePolygonGroups( const TArray<FPolygonGroupID>& PolygonG
 		for( int32 Index = PolygonGroupIDs.Num() - 1; Index >= 0; --Index )
 		{
 			const FPolygonGroupID PolygonGroupID = PolygonGroupIDs[ Index ];
-			const FMeshPolygonGroup& PolygonGroup = PolygonGroups[ PolygonGroupID.GetValue() ];
+			const FMeshPolygonGroup& PolygonGroup = PolygonGroups[ PolygonGroupID ];
 
 			RevertInput.PolygonGroupsToCreate.Emplace();
 			FPolygonGroupToCreate& PolygonGroupToCreate = RevertInput.PolygonGroupsToCreate.Last();
 
 			PolygonGroupToCreate.OriginalPolygonGroupID = PolygonGroupID;
-			PolygonGroupToCreate.Material = PolygonGroup.Material;
+			PolygonGroupToCreate.MaterialAsset = PolygonGroup.MaterialAsset.ToString();
+			PolygonGroupToCreate.MaterialSlotName = PolygonGroup.MaterialSlotName;
+			PolygonGroupToCreate.ImportedMaterialSlotName = PolygonGroup.ImportedMaterialSlotName;
 			PolygonGroupToCreate.bCastShadow = PolygonGroup.bCastShadow;
 			PolygonGroupToCreate.bEnableCollision = PolygonGroup.bEnableCollision;
 		}
@@ -4921,10 +4741,10 @@ void UEditableMesh::DeletePolygonGroups( const TArray<FPolygonGroupID>& PolygonG
 	{
 		for( const FPolygonGroupID PolygonGroupID : PolygonGroupIDs )
 		{
-			const FMeshPolygonGroup& PolygonGroup = PolygonGroups[ PolygonGroupID.GetValue() ];
+			const FMeshPolygonGroup& PolygonGroup = PolygonGroups[ PolygonGroupID ];
 			check( PolygonGroup.Polygons.Num() == 0 );
 
-			PolygonGroups.RemoveAt( PolygonGroupID.GetValue() );
+			PolygonGroups.Remove( PolygonGroupID );
 		}
 	}
 }
@@ -5086,6 +4906,9 @@ void UEditableMesh::ChangePolygonsVertexInstances( const TArray<FChangeVertexIns
 {
 	UE_LOG( LogEditableMesh, Verbose, TEXT( "ChangePolygonsVertexInstances: %s" ), *LogHelpers::ArrayToString( VertexInstancesForPolygons ) );
 
+	FPolygonArray& Polygons = GetMeshDescription()->Polygons();
+	FVertexInstanceArray& VertexInstances = GetMeshDescription()->VertexInstances();
+
 	// Back everything up
 	{
 		FChangePolygonsVertexInstancesChangeInput RevertInput;
@@ -5096,7 +4919,7 @@ void UEditableMesh::ChangePolygonsVertexInstances( const TArray<FChangeVertexIns
 		{
 			const FChangeVertexInstancesForPolygon& VertexInstancesForPolygon = VertexInstancesForPolygons[ Index ];
 			const FPolygonID PolygonID = VertexInstancesForPolygon.PolygonID;
-			const FMeshPolygon& Polygon = Polygons[ PolygonID.GetValue() ];
+			const FMeshPolygon& Polygon = Polygons[ PolygonID ];
 
 			RevertInput.VertexInstancesForPolygons.Emplace();
 			FChangeVertexInstancesForPolygon& RevertVertexInstancesForPolygon = RevertInput.VertexInstancesForPolygons.Last();
@@ -5140,16 +4963,16 @@ void UEditableMesh::ChangePolygonsVertexInstances( const TArray<FChangeVertexIns
 		{
 			const FPolygonID PolygonID = VertexInstancesForPolygon.PolygonID;
 			PolygonIDs.Add( PolygonID );
-			FMeshPolygon& Polygon = Polygons[ PolygonID.GetValue() ];
+			FMeshPolygon& Polygon = Polygons[ PolygonID ];
 
 			for( const FVertexIndexAndInstanceID& IndexAndInstance : VertexInstancesForPolygon.PerimeterVertexIndicesAndInstanceIDs )
 			{
 				// Disconnect old vertex instance from polygon, and connect new one
 				const FVertexInstanceID OldVertexInstanceID = Polygon.PerimeterContour.VertexInstanceIDs[ IndexAndInstance.ContourIndex ];
-				FMeshVertexInstance& OldVertexInstance = VertexInstances[ OldVertexInstanceID.GetValue() ];
+				FMeshVertexInstance& OldVertexInstance = VertexInstances[ OldVertexInstanceID ];
 				verify( OldVertexInstance.ConnectedPolygons.Remove( PolygonID ) == 1 );
 
-				FMeshVertexInstance& NewVertexInstance = VertexInstances[ IndexAndInstance.VertexInstanceID.GetValue() ];
+				FMeshVertexInstance& NewVertexInstance = VertexInstances[ IndexAndInstance.VertexInstanceID ];
 				check( !NewVertexInstance.ConnectedPolygons.Contains( PolygonID ) );
 				NewVertexInstance.ConnectedPolygons.Add( PolygonID );
 
@@ -5175,10 +4998,10 @@ void UEditableMesh::ChangePolygonsVertexInstances( const TArray<FChangeVertexIns
 				{
 					// Disconnect old vertex instance from polygon, and connect new one
 					const FVertexInstanceID OldVertexInstanceID = Polygon.HoleContours[ HoleIndex ].VertexInstanceIDs[ IndexAndInstance.ContourIndex ];
-					FMeshVertexInstance& OldVertexInstance = VertexInstances[ OldVertexInstanceID.GetValue() ];
+					FMeshVertexInstance& OldVertexInstance = VertexInstances[ OldVertexInstanceID ];
 					verify( OldVertexInstance.ConnectedPolygons.Remove( PolygonID ) == 1 );
 
-					FMeshVertexInstance& NewVertexInstance = VertexInstances[ IndexAndInstance.VertexInstanceID.GetValue() ];
+					FMeshVertexInstance& NewVertexInstance = VertexInstances[ IndexAndInstance.VertexInstanceID ];
 					check( !NewVertexInstance.ConnectedPolygons.Contains( PolygonID ) );
 					NewVertexInstance.ConnectedPolygons.Add( PolygonID );
 
@@ -5214,11 +5037,11 @@ void UEditableMesh::ChangePolygonsVertexInstances( const TArray<FChangeVertexIns
 
 FVertexInstanceID UEditableMesh::GetVertexInstanceInPolygonForVertex( const FPolygonID PolygonID, const FVertexID VertexID ) const
 {
-	const FMeshVertex& Vertex = Vertices[ VertexID.GetValue() ];
+	const FVertexInstanceArray& VertexInstances = GetMeshDescription()->VertexInstances();
 
-	for( const FVertexInstanceID VertexInstanceID : Vertex.VertexInstanceIDs )
+	for( const FVertexInstanceID VertexInstanceID : GetMeshDescription()->GetVertex( VertexID ).VertexInstanceIDs )
 	{
-		const FMeshVertexInstance& VertexInstance = VertexInstances[ VertexInstanceID.GetValue() ];
+		const FMeshVertexInstance& VertexInstance = VertexInstances[ VertexInstanceID ];
 		if( VertexInstance.ConnectedPolygons.Contains( PolygonID ) )
 		{
 			return VertexInstanceID;
@@ -5233,11 +5056,9 @@ void UEditableMesh::GetConnectedSoftEdges( const FVertexID VertexID, TArray<FEdg
 {
 	OutConnectedSoftEdges.Reset();
 
-	const FMeshVertex& Vertex = Vertices[ VertexID.GetValue() ];
-	for( const FEdgeID ConnectedEdgeID : Vertex.ConnectedEdgeIDs )
+	for( const FEdgeID ConnectedEdgeID : GetMeshDescription()->GetVertex( VertexID ).ConnectedEdgeIDs )
 	{
-		const FMeshEdge& Edge = Edges[ ConnectedEdgeID.GetValue() ];
-		if( !Edge.bIsHardEdge )
+		if( GetEdgeAttribute( ConnectedEdgeID, UEditableMeshAttribute::EdgeIsHard(), 0 ).X == 0.0f )
 		{
 			OutConnectedSoftEdges.Add( ConnectedEdgeID );
 		}
@@ -5249,6 +5070,8 @@ void UEditableMesh::GetVertexInstancesInSameSoftEdgedGroup( const FVertexID Vert
 {
 	OutVertexInstanceIDs.Reset();
 
+	const FEdgeArray& Edges = GetMeshDescription()->Edges();
+
 	// Cache a list of all soft edges which share this vertex.
 	// We're not interested in hard edges as they denote a transition to a different soft edged group.
 	// Note that a vertex is usually only treated as 'hard' if it is met by two hard edges: if a hard edge terminates at the vertex we're looking at,
@@ -5288,7 +5111,7 @@ void UEditableMesh::GetVertexInstancesInSameSoftEdgedGroup( const FVertexID Vert
 			// have the current polygon as an adjacent.
 			for( const FEdgeID ConnectedSoftEdge : ConnectedSoftEdges )
 			{
-				const FMeshEdge& Edge = Edges[ ConnectedSoftEdge.GetValue() ];
+				const FMeshEdge& Edge = Edges[ ConnectedSoftEdge ];
 				if( Edge.ConnectedPolygons.Contains( PolygonToCheck ) )
 				{
 					for( const FPolygonID AdjacentPolygon : Edge.ConnectedPolygons )
@@ -5307,10 +5130,12 @@ void UEditableMesh::SetPolygonsVertexAttributes( const TArray<FVertexAttributesF
 {
 	UE_LOG( LogEditableMesh, Verbose, TEXT( "SetPolygonsVertexAttributes: %s" ), *LogHelpers::ArrayToString( VertexAttributesForPolygons ) );
 
+	FPolygonArray& Polygons = GetMeshDescription()->Polygons();
+
 	for( const FVertexAttributesForPolygon& VertexAttributesForPolygon : VertexAttributesForPolygons )
 	{
 		const FPolygonID PolygonID = VertexAttributesForPolygon.PolygonID;
-		FMeshPolygon& Polygon = Polygons[ PolygonID.GetValue() ];
+		FMeshPolygon& Polygon = Polygons[ PolygonID ];
 
 		SetPolygonContourVertexAttributes( Polygon.PerimeterContour, PolygonID, INDEX_NONE, VertexAttributesForPolygon.PerimeterVertexAttributeLists );
 
@@ -5361,8 +5186,7 @@ void UEditableMesh::SetPolygonContourVertexAttributes( FMeshPolygonContour& Cont
 
 		// Get vertex instance and vertex.
 		const FVertexInstanceID VertexInstanceID = Contour.VertexInstanceIDs[ Index ];
-		const FMeshVertexInstance& VertexInstance = VertexInstances[ VertexInstanceID.GetValue() ];
-		const FVertexID VertexID = VertexInstance.VertexID;
+		const FVertexID VertexID = GetVertexInstanceVertex( VertexInstanceID );
 
 		// We need to determine whether changing the attributes on this polygon vertex would require vertex instances to be split or merged.
 		// First determine all the existing vertex instances which are in this soft edged group, that is to say,
@@ -5816,9 +5640,14 @@ void UEditableMesh::TryToRemoveVertex( const FVertexID VertexID, bool& bOutWasVe
 
 		// Delete the vertex instances and subsequently orphaned vertex
 		{
-			FMeshVertex& Vertex = Vertices[ VertexID.GetValue() ];
+			FMeshVertex& Vertex = GetMeshDescription()->GetVertex( VertexID );
 			const bool bDeleteOrphanedVertices = true;
-			DeleteVertexInstances( Vertex.VertexInstanceIDs, bDeleteOrphanedVertices );
+
+			// Take a copy of the array, because it will be modified by the DeleteVertexInstances call
+			static TArray<FVertexInstanceID> VertexInstanceIDs;
+			VertexInstanceIDs = Vertex.VertexInstanceIDs;
+
+			DeleteVertexInstances( VertexInstanceIDs, bDeleteOrphanedVertices );
 		}
 
 		bOutWasVertexRemoved = true;
@@ -7010,7 +6839,7 @@ void UEditableMesh::BevelPolygons( const TArray<FPolygonID>& PolygonIDs, const f
 
 float UEditableMesh::GetPolygonCornerAngleForVertex( const FPolygonID PolygonID, const FVertexID VertexID ) const
 {
-	const FMeshPolygon& Polygon = Polygons[ PolygonID.GetValue() ];
+	const FMeshPolygon& Polygon = GetMeshDescription()->GetPolygon( PolygonID );
 
 	// Lambda function which returns the inner angle at a given index on a polygon contour
 	auto GetContourAngle = [ this ]( const FMeshPolygonContour& Contour, const int32 ContourIndex )
@@ -7027,9 +6856,10 @@ float UEditableMesh::GetPolygonCornerAngleForVertex( const FPolygonID PolygonID,
 		return FMath::Acos( FVector::DotProduct( Direction1, Direction2 ) );
 	};
 
-	auto IsVertexInstancedFromThisVertex = [ this, VertexID ]( const FVertexInstanceID VertexInstanceID )
+	const FVertexInstanceArray& VertexInstances = GetMeshDescription()->VertexInstances();
+	auto IsVertexInstancedFromThisVertex = [ &VertexInstances, VertexID ]( const FVertexInstanceID VertexInstanceID )
 	{
-		return VertexInstances[ VertexInstanceID.GetValue() ].VertexID == VertexID;
+		return VertexInstances[ VertexInstanceID ].VertexID == VertexID;
 	};
 
 	// First look for the vertex instance in the perimeter
@@ -7060,9 +6890,13 @@ float UEditableMesh::GetPolygonCornerAngleForVertex( const FPolygonID PolygonID,
 
 void UEditableMesh::GeneratePolygonTangentsAndNormals( const TArray<FPolygonID>& PolygonIDs )
 {
+	FPolygonArray& Polygons = GetMeshDescription()->Polygons();
+	const FVertexArray& Vertices = GetMeshDescription()->Vertices();
+	const FVertexInstanceArray& VertexInstances = GetMeshDescription()->VertexInstances();
+
 	for( const FPolygonID PolygonID : PolygonIDs )
 	{
-		FMeshPolygon& Polygon = Polygons[ PolygonID.GetValue() ];
+		FMeshPolygon& Polygon = Polygons[ PolygonID ];
 
 		// Calculate the tangent basis for the polygon, based on the average of all constituent triangles
 		FVector Normal = FVector::ZeroVector;
@@ -7071,13 +6905,13 @@ void UEditableMesh::GeneratePolygonTangentsAndNormals( const TArray<FPolygonID>&
 
 		for( const FMeshTriangle& Triangle : Polygon.Triangles )
 		{
-			const FMeshVertexInstance& VertexInstance0 = VertexInstances[ Triangle.VertexInstanceID0.GetValue() ];
-			const FMeshVertexInstance& VertexInstance1 = VertexInstances[ Triangle.VertexInstanceID1.GetValue() ];
-			const FMeshVertexInstance& VertexInstance2 = VertexInstances[ Triangle.VertexInstanceID2.GetValue() ];
+			const FMeshVertexInstance& VertexInstance0 = VertexInstances[ Triangle.VertexInstanceID0 ];
+			const FMeshVertexInstance& VertexInstance1 = VertexInstances[ Triangle.VertexInstanceID1 ];
+			const FMeshVertexInstance& VertexInstance2 = VertexInstances[ Triangle.VertexInstanceID2 ];
 
-			const FMeshVertex& Vertex0 = Vertices[ VertexInstance0.VertexID.GetValue() ];
-			const FMeshVertex& Vertex1 = Vertices[ VertexInstance1.VertexID.GetValue() ];
-			const FMeshVertex& Vertex2 = Vertices[ VertexInstance2.VertexID.GetValue() ];
+			const FMeshVertex& Vertex0 = Vertices[ VertexInstance0.VertexID ];
+			const FMeshVertex& Vertex1 = Vertices[ VertexInstance1.VertexID ];
+			const FMeshVertex& Vertex2 = Vertices[ VertexInstance2.VertexID ];
 
 			const FVector DPosition1 = Vertex1.VertexPosition - Vertex0.VertexPosition;
 			const FVector DPosition2 = Vertex2.VertexPosition - Vertex0.VertexPosition;
@@ -7112,9 +6946,10 @@ void UEditableMesh::GenerateTangentsAndNormals()
 	static TSet<FVertexInstanceID> VertexInstanceIDs;
 	VertexInstanceIDs.Reset();
 
+	const FPolygonArray& Polygons = GetMeshDescription()->Polygons();
 	for( const FPolygonID PolygonID : PolygonsPendingNewTangentBasis )
 	{
-		const FMeshPolygon& Polygon = Polygons[ PolygonID.GetValue() ];
+		const FMeshPolygon& Polygon = Polygons[ PolygonID ];
 
 		VertexInstanceIDs.Append( Polygon.PerimeterContour.VertexInstanceIDs );
 		for( const FMeshPolygonContour& HoleContour : Polygon.HoleContours )
@@ -7126,13 +6961,14 @@ void UEditableMesh::GenerateTangentsAndNormals()
 	static TArray<FAttributesForVertexInstance> AttributesForVertexInstances;
 	AttributesForVertexInstances.Reset( VertexInstanceIDs.Num() );
 
+	FVertexInstanceArray& VertexInstances = GetMeshDescription()->VertexInstances();
 	for( const FVertexInstanceID VertexInstanceID : VertexInstanceIDs )
 	{
 		FVector Normal = FVector::ZeroVector;
 		FVector Tangent = FVector::ZeroVector;
 		FVector Binormal = FVector::ZeroVector;
 
-		FMeshVertexInstance& VertexInstance = VertexInstances[ VertexInstanceID.GetValue() ];
+		FMeshVertexInstance& VertexInstance = VertexInstances[ VertexInstanceID ];
 		const FVertexID VertexID = VertexInstance.VertexID;
 
 		// Get all polygons connected to this vertex instance, and also any in the same smoothing group connected to a different vertex instance
@@ -7146,7 +6982,7 @@ void UEditableMesh::GenerateTangentsAndNormals()
 		{
 			float Angle = GetPolygonCornerAngleForVertex( ConnectedPolygonID, VertexID );
 
-			const FMeshPolygon& Polygon = Polygons[ ConnectedPolygonID.GetValue() ];
+			const FMeshPolygon& Polygon = Polygons[ ConnectedPolygonID ];
 			Normal += Polygon.PolygonNormal * Angle;
 
 			// If this polygon is actually connected to the vertex instance we're processing, also include its contributions towards the tangent
@@ -7441,6 +7277,9 @@ void UEditableMesh::SetEdgesVertices( const TArray<FVerticesForEdge>& VerticesFo
 {
 	UE_LOG( LogEditableMesh, Verbose, TEXT( "SetEdgesVertices: %s" ), *LogHelpers::ArrayToString( VerticesForEdges ) );
 
+	FEdgeArray& Edges = GetMeshDescription()->Edges();
+	FVertexArray& Vertices = GetMeshDescription()->Vertices();
+
 	FSetEdgesVerticesChangeInput RevertInput;
 	RevertInput.VerticesForEdges.AddUninitialized( VerticesForEdges.Num() );
 
@@ -7457,14 +7296,14 @@ void UEditableMesh::SetEdgesVertices( const TArray<FVerticesForEdge>& VerticesFo
 		GetEdgeVertices( VerticesForEdge.EdgeID, /* Out */ RevertVerticesForEdge.NewVertexID0, /* Out */ RevertVerticesForEdge.NewVertexID1 );
 
 		// Edit the edge
-		FMeshEdge& Edge = Edges[ VerticesForEdge.EdgeID.GetValue() ];
+		FMeshEdge& Edge = Edges[ VerticesForEdge.EdgeID ];
 		EdgeIDs[ EdgeNumber ] = VerticesForEdge.EdgeID;
 
 		for( uint32 EdgeVertexNumber = 0; EdgeVertexNumber < 2; ++EdgeVertexNumber )
 		{
 			// Disconnect the edge from its existing vertices
 			const FVertexID VertexID = Edge.VertexIDs[ EdgeVertexNumber ];
-			FMeshVertex& Vertex = Vertices[ VertexID.GetValue() ];
+			FMeshVertex& Vertex = Vertices[ VertexID ];
 			verify( Vertex.ConnectedEdgeIDs.RemoveSingleSwap( VerticesForEdge.EdgeID ) == 1 );	// Must have been already connected!
 		}
 
@@ -7475,7 +7314,7 @@ void UEditableMesh::SetEdgesVertices( const TArray<FVerticesForEdge>& VerticesFo
 		for( uint32 EdgeVertexNumber = 0; EdgeVertexNumber < 2; ++EdgeVertexNumber )
 		{
 			const FVertexID VertexID = Edge.VertexIDs[ EdgeVertexNumber ];
-			FMeshVertex& Vertex = Vertices[ VertexID.GetValue() ];
+			FMeshVertex& Vertex = Vertices[ VertexID ];
 
 			check( !Vertex.ConnectedEdgeIDs.Contains( VerticesForEdge.EdgeID ) );	// Should not have already been connected
 			Vertex.ConnectedEdgeIDs.Add( VerticesForEdge.EdgeID );
@@ -7500,7 +7339,8 @@ void UEditableMesh::InsertPolygonPerimeterVertices( const FPolygonID PolygonID,
 	// @todo mesheditor: see if we want to keep this action. It is only used by SplitEdge, and doesn't generate missing edges.
 	// We can achieve the same thing by deleting a polygon without deleting vertex instances, and creating a new polygon with an extra vertex instance.
 
-	FMeshPolygon& Polygon = Polygons[ PolygonID.GetValue() ];
+	FMeshPolygon& Polygon = GetMeshDescription()->GetPolygon( PolygonID );
+	FVertexInstanceArray& VertexInstances = GetMeshDescription()->VertexInstances();
 
 	// @todo mesheditor: should create missing edges here too?
 
@@ -7513,7 +7353,7 @@ void UEditableMesh::InsertPolygonPerimeterVertices( const FPolygonID PolygonID,
 		Polygon.PerimeterContour.VertexInstanceIDs.Insert( VertexInstanceID, InsertBeforeVertexNumber + InsertVertexIter );
 
 		// Connect polygon to vertex instance
-		FMeshVertexInstance& VertexInstance = VertexInstances[ VertexInstanceID.GetValue() ];
+		FMeshVertexInstance& VertexInstance = VertexInstances[ VertexInstanceID ];
 		check( !VertexInstance.ConnectedPolygons.Contains( PolygonID ) );
 		VertexInstance.ConnectedPolygons.Add( PolygonID );
 	}
@@ -7553,7 +7393,7 @@ void UEditableMesh::InsertPolygonPerimeterVertices( const FPolygonID PolygonID,
 void UEditableMesh::RemovePolygonPerimeterVertices( const FPolygonID PolygonID, const int32 FirstVertexNumberToRemove, const int32 NumVerticesToRemove, const bool bDeleteOrphanedVertexInstances )
 {
 	UE_LOG( LogEditableMesh, Verbose, TEXT( "RemovePolygonPerimeterVertices: PolygonID:%s FirstVertexNumberToRemove:%d NumVerticesToRemove:%d" ), *PolygonID.ToString(), FirstVertexNumberToRemove, NumVerticesToRemove );
-	FMeshPolygon& Polygon = Polygons[ PolygonID.GetValue() ];
+	FMeshPolygon& Polygon = GetMeshDescription()->GetPolygon( PolygonID );
 
 	// Back up 
 	{
@@ -7577,6 +7417,8 @@ void UEditableMesh::RemovePolygonPerimeterVertices( const FPolygonID PolygonID,
 		static TArray<FVertexInstanceID> OrphanedVertexInstanceIDs;
 		OrphanedVertexInstanceIDs.Reset();
 
+		FVertexInstanceArray& VertexInstances = GetMeshDescription()->VertexInstances();
+
 		// Delete them backwards so it is done in the opposite order to Insert
 		for( int32 Index = NumVerticesToRemove - 1; Index >= 0; --Index )
 		{
@@ -7584,7 +7426,7 @@ void UEditableMesh::RemovePolygonPerimeterVertices( const FPolygonID PolygonID,
 			const FVertexInstanceID ContourVertexInstanceID = Polygon.PerimeterContour.VertexInstanceIDs[ VertexNumber ];
 
 			// Disconnect the polygon from the vertex instance
-			FMeshVertexInstance& VertexInstance = VertexInstances[ ContourVertexInstanceID.GetValue() ];
+			FMeshVertexInstance& VertexInstance = VertexInstances[ ContourVertexInstanceID ];
 			verify( VertexInstance.ConnectedPolygons.Remove( PolygonID ) == 1 );
 
 			// If the vertex instance is now orphaned, add it to the list
@@ -7840,104 +7682,68 @@ void UEditableMesh::TriangulatePolygons( const TArray<FPolygonID>& PolygonIDs, T
 }
 
 
-void UEditableMesh::AssignMaterialToPolygons( const TArray<FPolygonID>& PolygonIDs, UMaterialInterface* Material, TArray<FPolygonID>& NewPolygonIDs )
+void UEditableMesh::AssignPolygonsToPolygonGroups( const TArray<FPolygonGroupForPolygon>& PolygonGroupForPolygons, const bool bDeleteOrphanedPolygonGroups )
 {
-	// @todo mesheditor: create more generic method to assign a polygon to a new polygon group
-
-	// We need to move polygons from one polygon group to another. First find the new polygon group id to move them to.
-	const bool bCreateNewPolygonGroupIfNotFound = true;
-	const FPolygonGroupID NewPolygonGroupID = GetPolygonGroupIDFromMaterial( Material, bCreateNewPolygonGroupIfNotFound );
-	check( NewPolygonGroupID != FPolygonGroupID::Invalid );
-
-	// Make an array of FPolygonToCreate, based on the ones we wish to move. Everything is the same, other than the PolygonGroupID
-	static TArray<FPolygonToCreate> PolygonsToCreate;
-	PolygonsToCreate.Reset( PolygonIDs.Num() );
-
-	for( int32 PolygonNumber = 0; PolygonNumber < PolygonIDs.Num(); ++PolygonNumber )
+	// Back up
 	{
-		const FPolygonID OriginalPolygonID = PolygonIDs[ PolygonNumber ];
-		const FMeshPolygon& Polygon = Polygons[ OriginalPolygonID.GetValue() ];
-
-		FPolygonToCreate& PolygonToCreate = *new( PolygonsToCreate ) FPolygonToCreate();
-
-		PolygonToCreate.PolygonGroupID = NewPolygonGroupID;
-		PolygonToCreate.OriginalPolygonID = FPolygonID::Invalid;
-
-		const int32 PerimeterVertexCount = GetPolygonPerimeterVertexCount( OriginalPolygonID );
-		PolygonToCreate.PerimeterVertices.Reserve( PerimeterVertexCount );
-
-		for( int32 VertexNumber = 0; VertexNumber < PerimeterVertexCount; ++VertexNumber )
+		FAssignPolygonsToPolygonGroupChangeInput RevertInput;
+		RevertInput.PolygonGroupForPolygons.Reserve( PolygonGroupForPolygons.Num() );
+		for( int32 Index = PolygonGroupForPolygons.Num() - 1; Index >= 0; --Index )
 		{
-			FVertexAndAttributes& PerimeterVertex = *new( PolygonToCreate.PerimeterVertices ) FVertexAndAttributes();
-
-			PerimeterVertex.VertexID = GetPolygonPerimeterVertex( OriginalPolygonID, VertexNumber );
-
-			TArray<FMeshElementAttributeData>& PerimeterVertexAttributeList = PerimeterVertex.PolygonVertexAttributes.Attributes;
-
-			for( const FName AttributeName : UEditableMesh::GetValidVertexInstanceAttributes() )
-			{
-				const int32 MaxAttributeIndex = GetMaxAttributeIndex( AttributeName );
-				for( int32 AttributeIndex = 0; AttributeIndex < MaxAttributeIndex; ++AttributeIndex )
-				{
-					FMeshElementAttributeData& PolygonVertexAttribute = *new( PerimeterVertexAttributeList ) FMeshElementAttributeData(
-						AttributeName,
-						AttributeIndex,
-						GetPolygonPerimeterVertexAttribute( OriginalPolygonID, VertexNumber, AttributeName, AttributeIndex ) );
-				}
-			}
+			RevertInput.PolygonGroupForPolygons.Emplace();
+			FPolygonGroupForPolygon& RevertPolygonGroupForPolygon = RevertInput.PolygonGroupForPolygons.Last();
+			RevertPolygonGroupForPolygon.PolygonID = PolygonGroupForPolygons[ Index ].PolygonID;
+			RevertPolygonGroupForPolygon.PolygonGroupID = GetGroupForPolygon( PolygonGroupForPolygons[ Index ].PolygonID );
 		}
 
-		// @todo mesheditor: Handle holes?
+		AddUndo( MakeUnique<FAssignPolygonsToPolygonGroupChange>( MoveTemp( RevertInput ) ) );
 	}
 
-	// Create the new polygons in the new PolygonGroup.
+	// Let the adapters handle this
+	for( UEditableMeshAdapter* Adapter : Adapters )
 	{
-		static TArray<FEdgeID> NewEdgeIDs;
-		CreatePolygons( PolygonsToCreate, NewPolygonIDs, NewEdgeIDs );
+		Adapter->OnAssignPolygonsToPolygonGroups( this, PolygonGroupForPolygons );
 	}
 
-	// Delete the polygons from the old PolygonGroup
-	{
-		const bool bDeleteOrphanedEdges = false;
-		const bool bDeleteOrphanedVertices = false;
-		const bool bDeleteOrphanedVertexInstances = false;
-		const bool bDeleteEmptyPolygonGroups = true;
-		DeletePolygons( PolygonIDs, bDeleteOrphanedEdges, bDeleteOrphanedVertices, bDeleteOrphanedVertexInstances, bDeleteEmptyPolygonGroups );
-	}
-}
+	// Maintain an array of orphaned polygon groups
+	static TArray<FPolygonGroupID> PolygonGroupsToDelete;
+	PolygonGroupsToDelete.Reset();
 
+	FPolygonArray& Polygons = GetMeshDescription()->Polygons();
+	FPolygonGroupArray& PolygonGroups = GetMeshDescription()->PolygonGroups();
 
-FPolygonGroupID UEditableMesh::GetPolygonGroupIDFromMaterial( UMaterialInterface* Material, bool bCreateNewSectionIfNotFound )
-{
-	// Iterate through the sections sparse array looking for an entry whose material index matches.
-	for( TSparseArray<FMeshPolygonGroup>::TConstIterator It( PolygonGroups ); It; ++It )
+	for( const FPolygonGroupForPolygon& PolygonGroupForPolygon : PolygonGroupForPolygons )
 	{
-		const FMeshPolygonGroup& PolygonGroup = *It;
+		const FPolygonID PolygonID = PolygonGroupForPolygon.PolygonID;
+		const FPolygonGroupID NewPolygonGroupID = PolygonGroupForPolygon.PolygonGroupID;
 
-		if( PolygonGroup.Material == Material )
+		FMeshPolygon& Polygon = Polygons[ PolygonID ];
+		const FPolygonGroupID OldPolygonGroupID = Polygon.PolygonGroupID;
+
+		// Remove polygon reference from old group
+		FMeshPolygonGroup& OldPolygonGroup = PolygonGroups[ OldPolygonGroupID ];
+		verify( OldPolygonGroup.Polygons.RemoveSwap( PolygonID ) == 1 );
+
+		// If old group is now empty (and we're deleting orphans), add it to the list to delete
+		if( bDeleteOrphanedPolygonGroups && OldPolygonGroup.Polygons.Num() == 0 )
 		{
-			return FPolygonGroupID( It.GetIndex() );
+			PolygonGroupsToDelete.Add( OldPolygonGroupID );
 		}
-	}
-
-	// If we got here, the material index does not yet have a matching section.
-	if( bCreateNewSectionIfNotFound )
-	{
-		static TArray<FPolygonGroupToCreate> PolygonGroupsToCreate;
-		PolygonGroupsToCreate.Reset( 1 );
-		PolygonGroupsToCreate.Emplace();
-		FPolygonGroupToCreate& PolygonGroupToCreate = PolygonGroupsToCreate[ 0 ];
 
-		PolygonGroupToCreate.Material = Material;
-		PolygonGroupToCreate.bEnableCollision = true;
-		PolygonGroupToCreate.bCastShadow = true;
+		// Add polygon reference to new group
+		FMeshPolygonGroup& NewPolygonGroup = PolygonGroups[ NewPolygonGroupID ];
+		check( !NewPolygonGroup.Polygons.Contains( PolygonID ) );
+		NewPolygonGroup.Polygons.Add( PolygonID );
 
-		static TArray<FPolygonGroupID> PolygonGroupIDs;
-		CreatePolygonGroups( PolygonGroupsToCreate, PolygonGroupIDs );
-		return PolygonGroupIDs[ 0 ];
+		// Change polygon group
+		Polygon.PolygonGroupID = NewPolygonGroupID;
 	}
 
-	return FPolygonGroupID::Invalid;
+	// If there's any groups to delete, do it now
+	if( PolygonGroupsToDelete.Num() > 0 )
+	{
+		DeletePolygonGroups( PolygonGroupsToDelete );
+	}
 }
 
 
@@ -8636,16 +8442,11 @@ void UEditableMesh::SetTextureCoordinateCount( const int32 NumTexCoords )
 void UEditableMesh::QuadrangulateMesh( TArray<FPolygonID>& OutNewPolygonIDs )
 {
 	// Iterate through all polygon groups in the mesh, quadrangulating each in turn
-	const int32 MaxPolygonGroupIndex = GetPolygonGroupArraySize();
-	for( int32 PolygonGroupIndex = 0; PolygonGroupIndex < MaxPolygonGroupIndex; ++PolygonGroupIndex )
+	for( const FPolygonGroupID PolygonGroupID : GetMeshDescription()->PolygonGroups().GetElementIDs() )
 	{
-		const FPolygonGroupID PolygonGroupID( PolygonGroupIndex );
-		if( IsValidPolygonGroup( PolygonGroupID ) )
-		{
-			static TArray<FPolygonID> NewPolygonIDs;
-			QuadrangulatePolygonGroup( PolygonGroupID, NewPolygonIDs );
-			OutNewPolygonIDs.Append( NewPolygonIDs );
-		}
+		static TArray<FPolygonID> NewPolygonIDs;
+		QuadrangulatePolygonGroup( PolygonGroupID, NewPolygonIDs );
+		OutNewPolygonIDs.Append( NewPolygonIDs );
 	}
 }
 
@@ -9297,28 +9098,23 @@ void UEditableMesh::RebuildOctree()
 		// @todo mesheditor spatial: We'll currently only ever find vertices and edges that are connected to polygons.  Our spatial
 		// database only tracks polygons.  Do we need to support hover/selection of "loose" vertices and edges also?
 
-		const uint32 PolygonArraySize = GetPolygonArraySize();
-		for( uint32 PolygonIndex = 0; PolygonIndex < PolygonArraySize; ++PolygonIndex )
+		for( const FPolygonID PolygonID : GetMeshDescription()->Polygons().GetElementIDs() )
 		{
-			const FPolygonID PolygonID( PolygonIndex );
-			if( IsValidPolygon( PolygonID ) )
+			FBox BoundingBox;
 			{
-				FBox BoundingBox;
-				{
-					BoundingBox.Init();
+				BoundingBox.Init();
 
-					static TArray<FVertexID> PerimeterVertices;
-					GetPolygonPerimeterVertices( PolygonID, /* Out */ PerimeterVertices );
-					for( const FVertexID VertexID : PerimeterVertices )
-					{
-						const FVector VertexPosition = GetVertexAttribute( VertexID, UEditableMeshAttribute::VertexPosition(), 0 );
-						BoundingBox += VertexPosition;
-					}
+				static TArray<FVertexID> PerimeterVertices;
+				GetPolygonPerimeterVertices( PolygonID, /* Out */ PerimeterVertices );
+				for( const FVertexID VertexID : PerimeterVertices )
+				{
+					const FVector VertexPosition = GetVertexAttribute( VertexID, UEditableMeshAttribute::VertexPosition(), 0 );
+					BoundingBox += VertexPosition;
 				}
-
-				const FEditableMeshOctreePolygon OctreePolygon( *this, PolygonID, FBoxCenterAndExtent( BoundingBox ) );
-				Octree->AddElement( OctreePolygon );
 			}
+
+			const FEditableMeshOctreePolygon OctreePolygon( *this, PolygonID, FBoxCenterAndExtent( BoundingBox ) );
+			Octree->AddElement( OctreePolygon );
 		}
 
 		// @todo mesheditor debug
