- Split off mesh description into its own class and engine module. The intention is this will be shared by other mesh objects as the source mesh format, rather than being part of UEditableMesh.

diff --git a/Engine/Plugins/Runtime/EditableMesh/Source/EditableMesh/EditableMeshChanges.cpp b/Engine/Plugins/Runtime/EditableMesh/Source/EditableMesh/EditableMeshChanges.cpp
--- a/Engine/Plugins/Runtime/EditableMesh/Source/EditableMesh/EditableMeshChanges.cpp
+++ b/Engine/Plugins/Runtime/EditableMesh/Source/EditableMesh/EditableMeshChanges.cpp
@@ -374,7 +374,7 @@ TUniquePtr<FChange> FCreatePolygonGroupsChange::Execute( UObject* Object )
 FString FCreatePolygonGroupsChange::ToString() const
 {
 	return FString::Printf(
-		TEXT( "Create Section [PolygonGroupsToCreate:%s]" ),
+		TEXT( "Create PolygonGroups [PolygonGroupsToCreate:%s]" ),
 		*LogHelpers::ArrayToString( Input.PolygonGroupsToCreate ) );
 }
 
@@ -392,8 +392,25 @@ TUniquePtr<FChange> FDeletePolygonGroupsChange::Execute( UObject* Object )
 FString FDeletePolygonGroupsChange::ToString() const
 {
 	return FString::Printf(
-		TEXT( "Delete Section [PolygonGroupIDs:%s]" ),
+		TEXT( "Delete PolygonGroups [PolygonGroupIDs:%s]" ),
 		*LogHelpers::ArrayToString( Input.PolygonGroupIDs ) );
 }
 
 
+TUniquePtr<FChange> FAssignPolygonsToPolygonGroupChange::Execute( UObject* Object )
+{
+	UEditableMesh* EditableMesh = CastChecked<UEditableMesh>( Object );
+	verify( !EditableMesh->AnyChangesToUndo() );
+
+	EditableMesh->AssignPolygonsToPolygonGroups( Input.PolygonGroupForPolygons, Input.bDeleteOrphanedPolygonGroups );
+	return EditableMesh->MakeUndo();
+}
+
+
+FString FAssignPolygonsToPolygonGroupChange::ToString() const
+{
+	return FString::Printf(
+		TEXT( "Assign Polygons To PolygonGroups [PolygonGroupForPolygons:%s]" ),
+		*LogHelpers::ArrayToString( Input.PolygonGroupForPolygons ) );
+}
+
diff --git a/Engine/Plugins/Runtime/EditableMesh/Source/EditableMesh/EditableMeshChanges.h b/Engine/Plugins/Runtime/EditableMesh/Source/EditableMesh/EditableMeshChanges.h
--- a/Engine/Plugins/Runtime/EditableMesh/Source/EditableMesh/EditableMeshChanges.h
+++ b/Engine/Plugins/Runtime/EditableMesh/Source/EditableMesh/EditableMeshChanges.h
@@ -812,6 +812,11 @@ public:
 	{
 	}
 
+	FCreatePolygonGroupsChange( FCreatePolygonGroupsChangeInput&& InitInput )
+		: Input( MoveTemp( InitInput ) )
+	{
+	}
+
 	// Parent class overrides
 	virtual TUniquePtr<FChange> Execute( UObject* Object ) override;
 	virtual FString ToString() const override;
@@ -848,6 +853,11 @@ public:
 	{
 	}
 
+	FDeletePolygonGroupsChange( FDeletePolygonGroupsChangeInput&& InitInput )
+		: Input( MoveTemp( InitInput ) )
+	{
+	}
+
 	// Parent class overrides
 	virtual TUniquePtr<FChange> Execute( UObject* Object ) override;
 	virtual FString ToString() const override;
@@ -858,3 +868,48 @@ private:
 	/** The data we need to make this change */
 	FDeletePolygonGroupsChangeInput Input;
 };
+
+
+struct FAssignPolygonsToPolygonGroupChangeInput
+{
+	/** The polygon group IDs to delete */
+	TArray<FPolygonGroupForPolygon> PolygonGroupForPolygons;
+
+	/** Whether we should delete a polygon group if it becomes orphaned */
+	bool bDeleteOrphanedPolygonGroups;
+
+	/** Default constructor */
+	FAssignPolygonsToPolygonGroupChangeInput()
+		: PolygonGroupForPolygons()
+		, bDeleteOrphanedPolygonGroups( false )
+	{
+	}
+};
+
+
+class FAssignPolygonsToPolygonGroupChange : public FChange
+{
+
+public:
+
+	/** Constructor */
+	FAssignPolygonsToPolygonGroupChange( const FAssignPolygonsToPolygonGroupChangeInput& InitInput )
+		: Input( InitInput )
+	{
+	}
+
+	FAssignPolygonsToPolygonGroupChange( FAssignPolygonsToPolygonGroupChangeInput&& InitInput )
+		: Input( MoveTemp( InitInput ) )
+	{
+	}
+
+	// Parent class overrides
+	virtual TUniquePtr<FChange> Execute( UObject* Object ) override;
+	virtual FString ToString() const override;
+
+
+private:
+
+	/** The data we need to make this change */
+	FAssignPolygonsToPolygonGroupChangeInput Input;
+};
diff --git a/Engine/Plugins/Runtime/EditableMesh/Source/EditableMesh/EditableStaticMeshAdapter.cpp b/Engine/Plugins/Runtime/EditableMesh/Source/EditableMesh/EditableStaticMeshAdapter.cpp
--- a/Engine/Plugins/Runtime/EditableMesh/Source/EditableMesh/EditableStaticMeshAdapter.cpp
+++ b/Engine/Plugins/Runtime/EditableMesh/Source/EditableMesh/EditableStaticMeshAdapter.cpp
@@ -100,6 +100,9 @@ void UEditableStaticMeshAdapter::InitEditableStaticMesh( UEditableMesh* Editable
 				const bool bHasColor = StaticMeshLOD.ColorVertexBuffer.GetNumVertices() > 0;
 				check( !bHasColor || StaticMeshLOD.ColorVertexBuffer.GetNumVertices() == StaticMeshLOD.VertexBuffer.GetNumVertices() );
 
+				FVertexArray& Vertices = EditableMesh->GetMeshDescription()->Vertices();
+				FVertexInstanceArray& VertexInstances = EditableMesh->GetMeshDescription()->VertexInstances();
+
 				// @todo mesheditor cleanup: This code is very similar to the static mesh build code; try to share helper structs
 				static TMultiMap<int32, int32> OverlappingRenderingVertexIndices;
 				{
@@ -177,6 +180,7 @@ void UEditableStaticMeshAdapter::InitEditableStaticMesh( UEditableMesh* Editable
 				for( int32 RenderingVertexIndex = 0; RenderingVertexIndex < NumRenderingVertices; ++RenderingVertexIndex )
 				{
 					const FVector VertexPosition = StaticMeshLOD.PositionVertexBuffer.VertexPosition( RenderingVertexIndex );
+					const FVertexInstanceID VertexInstanceID( RenderingVertexIndex );
 
 					// Check to see if we already have this vertex
 					bool bAlreadyHaveVertexForPosition = false;
@@ -188,20 +192,20 @@ void UEditableStaticMeshAdapter::InitEditableStaticMesh( UEditableMesh* Editable
 						for( int32 OverlappingRenderingVertexIter = 0; OverlappingRenderingVertexIter < ThisRenderingVertexOverlaps.Num(); ++OverlappingRenderingVertexIter )
 						{
 							const int32 OverlappingRenderingVertexIndex = ThisRenderingVertexOverlaps[ OverlappingRenderingVertexIter ];
+							const FVertexInstanceID OverlappingVertexInstanceID( OverlappingRenderingVertexIndex );
 
 							// If the overlapping vertex instance index is smaller than our current index, we can safely assume that
 							// we've already processed this vertex position and created an editable mesh vertex for it.
 							if( OverlappingRenderingVertexIndex < RenderingVertexIndex )
 							{
-								check( EditableMesh->VertexInstances.IsAllocated( OverlappingRenderingVertexIndex ) );
-								const FVertexID ExistingVertexID = EditableMesh->VertexInstances[ OverlappingRenderingVertexIndex ].VertexID;
-								FMeshVertex& ExistingVertex = EditableMesh->Vertices[ ExistingVertexID.GetValue() ];
+								check( VertexInstances.IsValid( OverlappingVertexInstanceID ) );
+								const FVertexID ExistingVertexID = VertexInstances[ OverlappingVertexInstanceID ].VertexID;
+								FMeshVertex& ExistingVertex = Vertices[ ExistingVertexID ];
 
 								// We already have a unique editable vertex for this vertex instance position, so link them!
-								EditableMesh->VertexInstances.Insert( RenderingVertexIndex, FMeshVertexInstance() );
-								EditableMesh->VertexInstances[ RenderingVertexIndex ].VertexID = ExistingVertexID;
+								VertexInstances.Insert( VertexInstanceID );
+								VertexInstances[ VertexInstanceID ].VertexID = ExistingVertexID;
 
-								const FVertexInstanceID VertexInstanceID( RenderingVertexIndex );
 								checkSlow( !ExistingVertex.VertexInstanceIDs.Contains( VertexInstanceID ) );
 								ExistingVertex.VertexInstanceIDs.Add( VertexInstanceID );
 								bAlreadyHaveVertexForPosition = true;
@@ -213,20 +217,19 @@ void UEditableStaticMeshAdapter::InitEditableStaticMesh( UEditableMesh* Editable
 
 					if( !bAlreadyHaveVertexForPosition )
 					{
-						const FVertexID NewVertexID( EditableMesh->Vertices.Add( FMeshVertex() ) );
-						FMeshVertex& NewVertex = EditableMesh->Vertices[ NewVertexID.GetValue() ];
+						const FVertexID NewVertexID = Vertices.Add();
+						FMeshVertex& NewVertex = Vertices[ NewVertexID ];
 						NewVertex.VertexPosition = VertexPosition;
 						NewVertex.CornerSharpness = 0.0f;
 
-						EditableMesh->VertexInstances.Insert( RenderingVertexIndex, FMeshVertexInstance() );
-						EditableMesh->VertexInstances[ RenderingVertexIndex ].VertexID = NewVertexID;
+						VertexInstances.Insert( VertexInstanceID );
+						VertexInstances[ VertexInstanceID ].VertexID = NewVertexID;
 						
 						// @todo mesheditor: If a mesh somehow contained vertex instances that no triangle was referencing, this would cause
 						// the vertex instance to be ignored by the editable mesh code.  It would just sit in the vertex buffer (and in the
 						// editable mesh vertex's RenderingVertexIndices list), but would never be touched.  The editable mesh code only
 						// creates vertex instances for vertices that are attached to polygons, so this should never happen with meshes
 						// that we create and save.  Only if the incoming data had orphan vertices in it.  Should hopefully not be a problem.
-						const FVertexInstanceID VertexInstanceID( RenderingVertexIndex );
 						NewVertex.VertexInstanceIDs.Add( VertexInstanceID );
 
 						// NOTE: The new vertex's connected polygons will be filled in down below, as we're processing mesh triangles
@@ -234,7 +237,7 @@ void UEditableStaticMeshAdapter::InitEditableStaticMesh( UEditableMesh* Editable
 
 					// Populate the vertex instance attributes
 					{
-						FMeshVertexInstance& VertexInstance = EditableMesh->VertexInstances[ RenderingVertexIndex ];
+						FMeshVertexInstance& VertexInstance = VertexInstances[ VertexInstanceID ];
 						VertexInstance.VertexUVs.Reserve( NumUVs );
 						for( int32 UVIndex = 0; UVIndex < NumUVs; ++UVIndex )
 						{
@@ -263,6 +266,10 @@ void UEditableStaticMeshAdapter::InitEditableStaticMesh( UEditableMesh* Editable
 				static TMultiMap<FEdgeID, TTuple<FVertexInstanceID, FVertexInstanceID>> EdgeToVertexInstancePair;
 				EdgeToVertexInstancePair.Reset();
 
+				FPolygonGroupArray& PolygonGroups = EditableMesh->GetMeshDescription()->PolygonGroups();
+				FPolygonArray& Polygons = EditableMesh->GetMeshDescription()->Polygons();
+				FEdgeArray& Edges = EditableMesh->GetMeshDescription()->Edges();
+
 				// Add all polygon groups from the static mesh sections
 				const uint32 NumSections = StaticMeshLOD.Sections.Num();
 				for( uint32 RenderingSectionIndex = 0; RenderingSectionIndex < NumSections; ++RenderingSectionIndex )
@@ -270,17 +277,20 @@ void UEditableStaticMeshAdapter::InitEditableStaticMesh( UEditableMesh* Editable
 					const FStaticMeshSection& RenderingSection = StaticMeshLOD.Sections[ RenderingSectionIndex ];
 
 					// Create a new polygon group
-					const FPolygonGroupID NewPolygonGroupID( EditableMesh->PolygonGroups.Add( FMeshPolygonGroup() ) );
-					FMeshPolygonGroup& NewPolygonGroup = EditableMesh->PolygonGroups[ NewPolygonGroupID.GetValue() ];
+					const FPolygonGroupID NewPolygonGroupID = PolygonGroups.Add();
+					FMeshPolygonGroup& NewPolygonGroup = PolygonGroups[ NewPolygonGroupID ];
 
-					NewPolygonGroup.Material = StaticMesh->GetMaterial( RenderingSection.MaterialIndex );
+					const FStaticMaterial& StaticMaterial = StaticMesh->StaticMaterials[ RenderingSection.MaterialIndex ];
+					NewPolygonGroup.MaterialAsset = StaticMaterial.MaterialInterface->GetPathName();
+					NewPolygonGroup.MaterialSlotName = StaticMaterial.MaterialSlotName;
+					NewPolygonGroup.ImportedMaterialSlotName = StaticMaterial.ImportedMaterialSlotName;
 					NewPolygonGroup.bEnableCollision = RenderingSection.bEnableCollision;
 					NewPolygonGroup.bCastShadow = RenderingSection.bCastShadow;
 
 					// Create a rendering polygon group for holding the triangulated data and references to the static mesh rendering section.
 					// This is indexed by the same FPolygonGroupID as the PolygonGroups.
-					RenderingPolygonGroups.Insert( NewPolygonGroupID.GetValue(), FRenderingPolygonGroup() );
-					FRenderingPolygonGroup& NewRenderingPolygonGroup = RenderingPolygonGroups[ NewPolygonGroupID.GetValue() ];
+					RenderingPolygonGroups.Insert( NewPolygonGroupID );
+					FRenderingPolygonGroup& NewRenderingPolygonGroup = RenderingPolygonGroups[ NewPolygonGroupID ];
 
 					const uint32 NumSectionTriangles = RenderingSection.NumTriangles;
 					NewRenderingPolygonGroup.Triangles.Reserve( NumSectionTriangles );
@@ -291,12 +301,12 @@ void UEditableStaticMeshAdapter::InitEditableStaticMesh( UEditableMesh* Editable
 					{
 						const uint32 RenderingTriangleFirstVertexIndex = FRenderingPolygonGroup::TriangleIndexToRenderingTriangleFirstIndex( RenderingSection, FTriangleID( SectionTriangleIndex ) );
 
-						uint32 TriangleRenderingVertexIndices[ 3 ];
+						FVertexInstanceID TriangleVertexInstanceIDs[ 3 ];
 						FVertexID TriangleVertexIDs[ 3 ];
 						for( uint32 TriangleVertexIndex = 0; TriangleVertexIndex < 3; ++TriangleVertexIndex )
 						{
-							TriangleRenderingVertexIndices[ TriangleVertexIndex ] = RenderingIndices[ RenderingTriangleFirstVertexIndex + TriangleVertexIndex ];
-							TriangleVertexIDs[ TriangleVertexIndex ] = EditableMesh->VertexInstances[ TriangleRenderingVertexIndices[ TriangleVertexIndex ] ].VertexID;
+							TriangleVertexInstanceIDs[ TriangleVertexIndex ] = FVertexInstanceID( RenderingIndices[ RenderingTriangleFirstVertexIndex + TriangleVertexIndex ] );
+							TriangleVertexIDs[ TriangleVertexIndex ] = VertexInstances[ TriangleVertexInstanceIDs[ TriangleVertexIndex ] ].VertexID;
 						}
 
 						// Make sure we have a valid triangle.  The triangle can be invalid because at least two if it's vertex indices 
@@ -310,21 +320,22 @@ void UEditableStaticMeshAdapter::InitEditableStaticMesh( UEditableMesh* Editable
 						{
 							// Static meshes only support triangles, so there's no need to triangulate anything yet.  We'll make both
 							// a triangle and a polygon here.
-							const int32 NewTriangleIndex = SectionTriangleIndex;
+							const FTriangleID NewTriangleID = FTriangleID( SectionTriangleIndex );
 
-							NewRenderingPolygonGroup.Triangles.InsertUninitialized( NewTriangleIndex );
-							FMeshTriangle& NewTriangle = NewRenderingPolygonGroup.Triangles[ NewTriangleIndex ];
+							NewRenderingPolygonGroup.Triangles.Insert( NewTriangleID );
+							FMeshTriangle& NewTriangle = NewRenderingPolygonGroup.Triangles[ NewTriangleID ];
 
 							// Insert a polygon into the mesh
-							const FPolygonID NewPolygonID( EditableMesh->Polygons.Add( FMeshPolygon() ) );
-							FMeshPolygon& NewPolygon = EditableMesh->Polygons[ NewPolygonID.GetValue() ];
+							const FPolygonID NewPolygonID = Polygons.Add();
+							FMeshPolygon& NewPolygon = Polygons[ NewPolygonID ];
 							NewPolygon.PolygonGroupID = NewPolygonGroupID;
 							NewPolygonGroup.Polygons.Add( NewPolygonID );
 
 							// Create a rendering polygon mirror, indexed by the same ID
-							RenderingPolygons.Insert( NewPolygonID.GetValue(), FRenderingPolygon() );
-							FRenderingPolygon& NewRenderingPolygon = RenderingPolygons[ NewPolygonID.GetValue() ];
-							NewRenderingPolygon.TriangulatedPolygonTriangleIndices.Add( FTriangleID( NewTriangleIndex ) );
+							RenderingPolygons.Insert( NewPolygonID );
+							FRenderingPolygon& NewRenderingPolygon = RenderingPolygons[ NewPolygonID ];
+							NewRenderingPolygon.PolygonGroupID = NewPolygonGroupID;
+							NewRenderingPolygon.TriangulatedPolygonTriangleIndices.Add( NewTriangleID );
 
 
 							// Static meshes don't support polygons with holes, so we always start out with only a perimeter contour per polygon
@@ -334,13 +345,12 @@ void UEditableStaticMeshAdapter::InitEditableStaticMesh( UEditableMesh* Editable
 							// Connect vertices
 							for( uint32 TriangleVertexIndex = 0; TriangleVertexIndex < 3; ++TriangleVertexIndex )
 							{
-								const uint32 RenderingVertexIndex = TriangleRenderingVertexIndices[ TriangleVertexIndex ];
 								const FVertexID VertexID = TriangleVertexIDs[ TriangleVertexIndex ];
-								const FVertexInstanceID VertexInstanceID = FVertexInstanceID( RenderingVertexIndex );
+								const FVertexInstanceID VertexInstanceID = TriangleVertexInstanceIDs[ TriangleVertexIndex ];
 
 								NewPolygon.PerimeterContour.VertexInstanceIDs.Add( VertexInstanceID );
 
-								FMeshVertexInstance& VertexInstance = EditableMesh->VertexInstances[ VertexInstanceID.GetValue() ];
+								FMeshVertexInstance& VertexInstance = VertexInstances[ VertexInstanceID ];
 								VertexInstance.ConnectedPolygons.Add( NewPolygonID );
 
 								// The triangle points to each of its three vertices
@@ -364,15 +374,13 @@ void UEditableStaticMeshAdapter::InitEditableStaticMesh( UEditableMesh* Editable
 								// Add the edges of this triangle
 								for( uint32 TriangleEdgeNumber = 0; TriangleEdgeNumber < 3; ++TriangleEdgeNumber )
 								{
-									uint32 EdgeRenderingVertexIndices[ 2 ];
-									EdgeRenderingVertexIndices[ 0 ] = RenderingIndices[ RenderingTriangleFirstVertexIndex + ( TriangleEdgeNumber + 0 ) % 3 ];
-									EdgeRenderingVertexIndices[ 1 ] = RenderingIndices[ RenderingTriangleFirstVertexIndex + ( TriangleEdgeNumber + 1 ) % 3 ];
+									FVertexInstanceID EdgeVertexInstanceIDs[ 2 ];
+									EdgeVertexInstanceIDs[ 0 ] = TriangleVertexInstanceIDs[ TriangleEdgeNumber ];
+									EdgeVertexInstanceIDs[ 1 ] = TriangleVertexInstanceIDs[ ( TriangleEdgeNumber + 1 ) % 3 ];
 
 									FVertexID EdgeVertexIDs[ 2 ];
-									EdgeVertexIDs[ 0 ] = EditableMesh->VertexInstances[ EdgeRenderingVertexIndices[ 0 ] ].VertexID;
-									EdgeVertexIDs[ 1 ] = EditableMesh->VertexInstances[ EdgeRenderingVertexIndices[ 1 ] ].VertexID;
-
-									FVertexInstanceID EdgeVertexInstanceIDs[ 2 ];
+									EdgeVertexIDs[ 0 ] = TriangleVertexIDs[ TriangleEdgeNumber ];
+									EdgeVertexIDs[ 1 ] = TriangleVertexIDs[ ( TriangleEdgeNumber + 1 ) % 3 ];
 
 									// Check to see if this edge already exists
 									bool bAlreadyHaveEdge = false;
@@ -381,8 +389,6 @@ void UEditableStaticMeshAdapter::InitEditableStaticMesh( UEditableMesh* Editable
 										FEdgeID* FoundEdgeIDPtr = UniqueEdgeToEdgeID.Find( Local::Make64BitValueForEdge( EdgeVertexIDs[ 0 ], EdgeVertexIDs[ 1 ] ) );
 										if( FoundEdgeIDPtr != nullptr )
 										{
-											EdgeVertexInstanceIDs[ 0 ] = FVertexInstanceID( EdgeRenderingVertexIndices[ 0 ] );
-											EdgeVertexInstanceIDs[ 1 ] = FVertexInstanceID( EdgeRenderingVertexIndices[ 1 ] );
 											bAlreadyHaveEdge = true;
 											EdgeID = *FoundEdgeIDPtr;
 										}
@@ -393,8 +399,7 @@ void UEditableStaticMeshAdapter::InitEditableStaticMesh( UEditableMesh* Editable
 
 											if( FoundEdgeIDPtr != nullptr )
 											{
-												EdgeVertexInstanceIDs[ 0 ] = FVertexInstanceID( EdgeRenderingVertexIndices[ 1 ] );
-												EdgeVertexInstanceIDs[ 1 ] = FVertexInstanceID( EdgeRenderingVertexIndices[ 0 ] );
+												Swap( EdgeVertexIDs[ 0 ], EdgeVertexIDs[ 1 ] );
 												bAlreadyHaveEdge = true;
 												EdgeID = *FoundEdgeIDPtr;
 											}
@@ -404,8 +409,8 @@ void UEditableStaticMeshAdapter::InitEditableStaticMesh( UEditableMesh* Editable
 									if( !bAlreadyHaveEdge )
 									{
 										// Create the new edge.  We'll connect it to its polygons later on.
-										EdgeID = FEdgeID( EditableMesh->Edges.Add( FMeshEdge() ) );
-										FMeshEdge& NewEdge = EditableMesh->Edges[ EdgeID.GetValue() ];
+										EdgeID = Edges.Add();
+										FMeshEdge& NewEdge = Edges[ EdgeID ];
 
 										NewEdge.VertexIDs[ 0 ] = EdgeVertexIDs[ 0 ];
 										NewEdge.VertexIDs[ 1 ] = EdgeVertexIDs[ 1 ];
@@ -413,20 +418,17 @@ void UEditableStaticMeshAdapter::InitEditableStaticMesh( UEditableMesh* Editable
 										NewEdge.CreaseSharpness = 0.0f;
 
 										UniqueEdgeToEdgeID.Add( Local::Make64BitValueForEdge( EdgeVertexIDs[ 0 ], EdgeVertexIDs[ 1 ] ), EdgeID );
-
-										EdgeVertexInstanceIDs[ 0 ] = FVertexInstanceID( EdgeRenderingVertexIndices[ 0 ] );
-										EdgeVertexInstanceIDs[ 1 ] = FVertexInstanceID( EdgeRenderingVertexIndices[ 1 ] );
 									}
 
 									// Each edge will point back to the polygon that its connected to.  Remember, an edge can be shared by multiple
 									// polygons, but usually its best if only shared by up to two.
-									FMeshEdge& Edge = EditableMesh->Edges[ EdgeID.GetValue() ];
+									FMeshEdge& Edge = Edges[ EdgeID ];
 
 									Edge.ConnectedPolygons.AddUnique( NewPolygonID );
 
 									// Connect the end vertices to the edge
-									EditableMesh->Vertices[ EdgeVertexIDs[ 0 ].GetValue() ].ConnectedEdgeIDs.AddUnique( EdgeID );
-									EditableMesh->Vertices[ EdgeVertexIDs[ 1 ].GetValue() ].ConnectedEdgeIDs.AddUnique( EdgeID );
+									Vertices[ EdgeVertexIDs[ 0 ] ].ConnectedEdgeIDs.AddUnique( EdgeID );
+									Vertices[ EdgeVertexIDs[ 1 ] ].ConnectedEdgeIDs.AddUnique( EdgeID );
 
 									// Determine whether the edge is hard by checking whether
 									if( !Edge.bIsHardEdge )
@@ -437,8 +439,8 @@ void UEditableStaticMeshAdapter::InitEditableStaticMesh( UEditableMesh* Editable
 										for( const TTuple<FVertexInstanceID, FVertexInstanceID>& VertexInstancePair : VertexInstancePairs )
 										{
 											// If either of the ends of the edge just added have different normals to any overlapping edge, mark it as a hard edge
-											if( EditableMesh->VertexInstances[ VertexInstancePair.Get<0>().GetValue() ].Normal != EditableMesh->VertexInstances[ EdgeVertexInstanceIDs[ 0 ].GetValue() ].Normal ||
-												EditableMesh->VertexInstances[ VertexInstancePair.Get<1>().GetValue() ].Normal != EditableMesh->VertexInstances[ EdgeVertexInstanceIDs[ 1 ].GetValue() ].Normal )
+											if( VertexInstances[ VertexInstancePair.Get<0>() ].Normal != VertexInstances[ EdgeVertexInstanceIDs[ 0 ] ].Normal ||
+												VertexInstances[ VertexInstancePair.Get<1>() ].Normal != VertexInstances[ EdgeVertexInstanceIDs[ 1 ] ].Normal )
 											{
 												Edge.bIsHardEdge = true;
 												break;
@@ -468,13 +470,11 @@ void UEditableStaticMeshAdapter::InitEditableStaticMesh( UEditableMesh* Editable
 	// Cache polygon tangent bases
 	static TArray<FPolygonID> PolygonIDs;
 	PolygonIDs.Reset();
-	for( int32 PolygonIndex = 0; PolygonIndex < EditableMesh->Polygons.GetMaxIndex(); ++PolygonIndex )
+	for( const FPolygonID PolygonID : EditableMesh->GetMeshDescription()->Polygons().GetElementIDs() )
 	{
-		if( EditableMesh->Polygons.IsAllocated( PolygonIndex ) )
-		{
-			PolygonIDs.Add( FPolygonID( PolygonIndex ) );
-		}
+		PolygonIDs.Add( PolygonID );
 	}
+
 	EditableMesh->GeneratePolygonTangentsAndNormals( PolygonIDs );
 
 #if 0
@@ -518,13 +518,15 @@ void UEditableStaticMeshAdapter::InitializeFromEditableMesh( const UEditableMesh
 	RenderingPolygons.Reset();
 	RenderingPolygonGroups.Reset();
 
+	FPolygonGroupArray& PolygonGroups = EditableMesh->GetMeshDescription()->PolygonGroups();
+	FPolygonArray& Polygons = EditableMesh->GetMeshDescription()->Polygons();
+
 	// Create all the required rendering polygon groups (initialized to 'empty', each with a unique rendering section index)
 	int32 RenderingSectionIndex = 0;
-	for( TSparseArray<FMeshPolygonGroup>::TConstIterator It( EditableMesh->PolygonGroups ); It; ++It )
+	for( const FPolygonGroupID PolygonGroupID : PolygonGroups.GetElementIDs() )
 	{
-		const int32 Index = It.GetIndex();
-		RenderingPolygonGroups.Insert( Index, FRenderingPolygonGroup() );
-		FRenderingPolygonGroup& RenderingPolygonGroup = RenderingPolygonGroups[ Index ];
+		RenderingPolygonGroups.Insert( PolygonGroupID );
+		FRenderingPolygonGroup& RenderingPolygonGroup = RenderingPolygonGroups[ PolygonGroupID ];
 		RenderingPolygonGroup.RenderingSectionIndex = RenderingSectionIndex;
 		RenderingPolygonGroup.MaxTriangles = 0;
 
@@ -532,36 +534,30 @@ void UEditableStaticMeshAdapter::InitializeFromEditableMesh( const UEditableMesh
 	}
 
 	// Go through all the polygons, adding their triangles to the rendering polygon group
-	for( TSparseArray<FMeshPolygon>::TConstIterator It( EditableMesh->Polygons ); It; ++It )
+	for( const FPolygonID PolygonID : Polygons.GetElementIDs() )
 	{
-		const int32 Index = It.GetIndex();
-		RenderingPolygons.Insert( Index, FRenderingPolygon() );
-		FRenderingPolygon& RenderingPolygon = RenderingPolygons[ Index ];
+		const FMeshPolygon& Polygon = Polygons[ PolygonID ];
 
-		const FPolygonGroupID PolygonGroupID( It->PolygonGroupID );
-		const FMeshPolygonGroup& PolygonGroup = EditableMesh->PolygonGroups[ PolygonGroupID.GetValue() ];
-		FRenderingPolygonGroup& RenderingPolygonGroup = RenderingPolygonGroups[ PolygonGroupID.GetValue() ];
+		const FPolygonGroupID PolygonGroupID = Polygon.PolygonGroupID;
+		const FMeshPolygonGroup& PolygonGroup = PolygonGroups[ PolygonGroupID ];
+		FRenderingPolygonGroup& RenderingPolygonGroup = RenderingPolygonGroups[ PolygonGroupID ];
 
-		for( const FMeshTriangle& Triangle : It->Triangles )
+		RenderingPolygons.Insert( PolygonID );
+		FRenderingPolygon& RenderingPolygon = RenderingPolygons[ PolygonID ];
+		RenderingPolygon.PolygonGroupID = PolygonGroupID;
+
+		for( const FMeshTriangle& Triangle : Polygon.Triangles )
 		{
-			const int32 TriangleIndex = RenderingPolygonGroup.Triangles.Add( Triangle );
-			RenderingPolygon.TriangulatedPolygonTriangleIndices.Add( FTriangleID( TriangleIndex ) );
+			const FTriangleID TriangleID = RenderingPolygonGroup.Triangles.Add( Triangle );
+			RenderingPolygon.TriangulatedPolygonTriangleIndices.Add( TriangleID );
 		}
 
-		RenderingPolygonGroup.MaxTriangles += It->Triangles.Num();
+		RenderingPolygonGroup.MaxTriangles += Polygon.Triangles.Num();
 	}
 
 }
 
 
-int32 GetStaticMeshMaterialIndex( const UStaticMesh* StaticMesh, const UMaterialInterface* MaterialInterface )
-{
-	return StaticMesh->StaticMaterials.IndexOfByPredicate(
-		[ MaterialInterface ]( const FStaticMaterial& StaticMaterial ) { return StaticMaterial.MaterialInterface == MaterialInterface; }
-	);
-}
-
-
 void UEditableStaticMeshAdapter::OnRebuildRenderMesh( const UEditableMesh* EditableMesh )
 {
 	// @todo mesheditor urgent subdiv: Saw some editable mesh corruption artifacts when testing subDs in VR
@@ -577,7 +573,9 @@ void UEditableStaticMeshAdapter::OnRebuildRenderMesh( const UEditableMesh* Edita
 	static TArray< uint32 > IndexBuffer;
 	IndexBuffer.Reset();
 
-	StaticMeshLOD.Sections.Empty( EditableMesh->PolygonGroups.Num() );
+	const FPolygonGroupArray& PolygonGroups = EditableMesh->GetMeshDescription()->PolygonGroups();
+
+	StaticMeshLOD.Sections.Empty( PolygonGroups.Num() );
 
 	bool bHasColor = false;
 
@@ -585,19 +583,17 @@ void UEditableStaticMeshAdapter::OnRebuildRenderMesh( const UEditableMesh* Edita
 	{
 		check( EditableMesh->GetSubdivisionCount() > 0 );
 
-		const int32 SectionCount = EditableMesh->SubdivisionLimitData.Sections.Num();
-
 		// @todo mesheditor subdiv: Only 2 UVs supported for now, just to avoid having to use a dynamic array per vertex; needs a new data layout, probably (SoA)
 		const int32 SubdivMeshTextureCoordinateCount = FMath::Min( EditableMesh->GetTextureCoordinateCount(), 2 );
 
 		// The Sections sparse array mirrors the SubdivisionLimitData sections array;
 		// iterate through it in parallel in order to get the material index and other section properties
-		TSparseArray< FMeshPolygonGroup >::TConstIterator PolygonGroupIt( EditableMesh->PolygonGroups );
-		check( EditableMesh->PolygonGroups.Num() == SectionCount );
+		check( PolygonGroups.Num() == EditableMesh->SubdivisionLimitData.Sections.Num() );
 
-		for( int32 SectionNumber = 0; SectionNumber < SectionCount; ++SectionNumber )
+		int32 SectionNumber = 0;
+		for( const FPolygonGroupID PolygonGroupID : PolygonGroups.GetElementIDs() )
 		{
-			const FMeshPolygonGroup& PolygonGroup = *PolygonGroupIt;
+			const FMeshPolygonGroup& PolygonGroup = PolygonGroups[ PolygonGroupID ];
 			const FSubdivisionLimitSection& SubdivisionSection = EditableMesh->SubdivisionLimitData.Sections[ SectionNumber ];
 
 			const int32 SectionTriangleCount = SubdivisionSection.SubdividedQuads.Num() * 2;
@@ -618,8 +614,10 @@ void UEditableStaticMeshAdapter::OnRebuildRenderMesh( const UEditableMesh* Edita
 			StaticMeshSection.MinVertexIndex = FirstSectionVertexIndex;
 			StaticMeshSection.MaxVertexIndex = FirstSectionVertexIndex + SectionTriangleCount * 3;
 
-			const int32 MaterialIndex = GetStaticMeshMaterialIndex( StaticMesh, PolygonGroup.Material );
+			const int32 MaterialIndex = StaticMesh->GetMaterialIndex( PolygonGroup.MaterialSlotName );
 			check( MaterialIndex != INDEX_NONE );
+			check( StaticMesh->StaticMaterials[ MaterialIndex ].MaterialInterface->GetPathName() == PolygonGroup.MaterialAsset.ToString() );
+
 			StaticMeshSection.MaterialIndex = MaterialIndex;
 			StaticMeshSection.bEnableCollision = PolygonGroup.bEnableCollision;
 			StaticMeshSection.bCastShadow = PolygonGroup.bCastShadow;
@@ -676,25 +674,29 @@ void UEditableStaticMeshAdapter::OnRebuildRenderMesh( const UEditableMesh* Edita
 				}
 			}
 
-			++PolygonGroupIt;
+			++SectionNumber;
 		}
 	}
 	else
 	{
+		const FVertexArray& Vertices = EditableMesh->GetMeshDescription()->Vertices();
+		const FVertexInstanceArray& VertexInstances = EditableMesh->GetMeshDescription()->VertexInstances();
+
 		// set up vertex buffer elements
-		StaticMeshBuildVertices.SetNum( EditableMesh->VertexInstances.GetMaxIndex() );
+		StaticMeshBuildVertices.SetNum( VertexInstances.GetArraySize() );
 
-		for( TSparseArray<FMeshVertexInstance>::TConstIterator It( EditableMesh->VertexInstances ); It; ++It )
+		for( const FVertexInstanceID VertexInstanceID : VertexInstances.GetElementIDs() )
 		{
-			if( It->Color != FColor::White )
+			const FMeshVertexInstance& VertexInstance = VertexInstances[ VertexInstanceID ];
+
+			if( VertexInstance.Color != FColor::White )
 			{
 				bHasColor = true;
 			}
 
-			FStaticMeshBuildVertex& StaticMeshVertex = StaticMeshBuildVertices[ It.GetIndex() ];
-			const FMeshVertexInstance& VertexInstance = *It;
+			FStaticMeshBuildVertex& StaticMeshVertex = StaticMeshBuildVertices[ VertexInstanceID.GetValue() ];
 
-			StaticMeshVertex.Position = EditableMesh->Vertices[ VertexInstance.VertexID.GetValue() ].VertexPosition;
+			StaticMeshVertex.Position = Vertices[ VertexInstance.VertexID ].VertexPosition;
 			StaticMeshVertex.TangentX = VertexInstance.Tangent;
 			StaticMeshVertex.TangentY = FVector::CrossProduct( VertexInstance.Normal, VertexInstance.Tangent ).GetSafeNormal() * VertexInstance.BinormalSign;
 			StaticMeshVertex.TangentZ = VertexInstance.Normal;
@@ -706,10 +708,11 @@ void UEditableStaticMeshAdapter::OnRebuildRenderMesh( const UEditableMesh* Edita
 		}
 
 		// Set up index buffer
-		for( TSparseArray< FMeshPolygonGroup >::TConstIterator It( EditableMesh->PolygonGroups ); It; ++It )
+
+		for( const FPolygonGroupID PolygonGroupID : PolygonGroups.GetElementIDs() )
 		{
-			const FMeshPolygonGroup& PolygonGroup = EditableMesh->PolygonGroups[ It.GetIndex() ];
-			FRenderingPolygonGroup& RenderingPolygonGroup = RenderingPolygonGroups[ It.GetIndex() ];
+			const FMeshPolygonGroup& PolygonGroup = PolygonGroups[ PolygonGroupID ];
+			FRenderingPolygonGroup& RenderingPolygonGroup = RenderingPolygonGroups[ PolygonGroupID ];
 
 			RenderingPolygonGroup.RenderingSectionIndex = StaticMeshLOD.Sections.Num();
 
@@ -718,11 +721,12 @@ void UEditableStaticMeshAdapter::OnRebuildRenderMesh( const UEditableMesh* Edita
 			FStaticMeshSection& StaticMeshSection = StaticMeshLOD.Sections.Last();
 
 			StaticMeshSection.FirstIndex = IndexBuffer.Num();
-			StaticMeshSection.NumTriangles = RenderingPolygonGroup.Triangles.GetMaxIndex();
-			check( RenderingPolygonGroup.Triangles.GetMaxIndex() <= RenderingPolygonGroup.MaxTriangles );
+			StaticMeshSection.NumTriangles = RenderingPolygonGroup.Triangles.GetArraySize();
+			check( RenderingPolygonGroup.Triangles.GetArraySize() <= RenderingPolygonGroup.MaxTriangles );
 
-			const int32 MaterialIndex = GetStaticMeshMaterialIndex( StaticMesh, PolygonGroup.Material );
+			const int32 MaterialIndex = StaticMesh->GetMaterialIndex( PolygonGroup.MaterialSlotName );
 			check( MaterialIndex != INDEX_NONE );
+			check( StaticMesh->StaticMaterials[ MaterialIndex ].MaterialInterface->GetPathName() == PolygonGroup.MaterialAsset.ToString() );
 			StaticMeshSection.MaterialIndex = MaterialIndex;
 			StaticMeshSection.bEnableCollision = PolygonGroup.bEnableCollision;
 			StaticMeshSection.bCastShadow = PolygonGroup.bCastShadow;
@@ -730,19 +734,20 @@ void UEditableStaticMeshAdapter::OnRebuildRenderMesh( const UEditableMesh* Edita
 
 			if( RenderingPolygonGroup.Triangles.Num() > 0 )
 			{
-				IndexBuffer.Reserve( IndexBuffer.Num() + RenderingPolygonGroup.Triangles.GetMaxIndex() * 3 );
+				IndexBuffer.Reserve( IndexBuffer.Num() + RenderingPolygonGroup.Triangles.GetArraySize() * 3 );
 				uint32 MinIndex = TNumericLimits< uint32 >::Max();
 				uint32 MaxIndex = TNumericLimits< uint32 >::Min();
 
 				// Find the first valid vertex instance index, so that we have a value we can use for our degenerates
 				check( RenderingPolygonGroup.Triangles.Num() > 0 );
-				const FVertexInstanceID FirstValidRenderingID = RenderingPolygonGroup.Triangles.CreateConstIterator()->GetVertexInstanceID( 0 );
+				const FVertexInstanceID FirstValidRenderingID = RenderingPolygonGroup.Triangles[ RenderingPolygonGroup.Triangles.GetFirstValidID() ].GetVertexInstanceID( 0 );
 
-				for( int32 TriangleIndex = 0; TriangleIndex < RenderingPolygonGroup.Triangles.GetMaxIndex(); ++TriangleIndex )
+				for( int32 TriangleIndex = 0; TriangleIndex < RenderingPolygonGroup.Triangles.GetArraySize(); ++TriangleIndex )
 				{
-					if( RenderingPolygonGroup.Triangles.IsAllocated( TriangleIndex ) )
+					const FTriangleID TriangleID( TriangleIndex );
+					if( RenderingPolygonGroup.Triangles.IsValid( TriangleID ) )
 					{
-						const FMeshTriangle& Triangle = RenderingPolygonGroup.Triangles[ TriangleIndex ];
+						const FMeshTriangle& Triangle = RenderingPolygonGroup.Triangles[ TriangleID ];
 						for( int32 TriVert = 0; TriVert < 3; ++TriVert )
 						{
 							const uint32 RenderingVertexIndex = Triangle.GetVertexInstanceID( TriVert ).GetValue();
@@ -765,7 +770,7 @@ void UEditableStaticMeshAdapter::OnRebuildRenderMesh( const UEditableMesh* Edita
 				// Add any index buffer padding.
 				// This can be necessary if we have just loaded an editable mesh which had a MaxTriangles count in the editable mesh section
 				// greater than the sparse array max size (i.e. an extra gap had been reserved for tris).
-				const int32 IndexBufferPadding = RenderingPolygonGroup.MaxTriangles - RenderingPolygonGroup.Triangles.GetMaxIndex();
+				const int32 IndexBufferPadding = RenderingPolygonGroup.MaxTriangles - RenderingPolygonGroup.Triangles.GetArraySize();
 				if( IndexBufferPadding > 0 )
 				{
 					IndexBuffer.AddZeroed( IndexBufferPadding * 3 );
@@ -827,8 +832,8 @@ void UEditableStaticMeshAdapter::Serialize( FArchive& Ar )
 
 	Ar.UsingCustomVersion( FEditableMeshCustomVersion::GUID );
 
-	SerializeSparseArray( Ar, RenderingPolygons );
-	SerializeSparseArray( Ar, RenderingPolygonGroups );
+	Ar << RenderingPolygons;
+	Ar << RenderingPolygonGroups;
 }
 
 
@@ -904,20 +909,21 @@ void UEditableStaticMeshAdapter::OnRebuildRenderMeshFinish( const UEditableMesh*
 
 void UEditableStaticMeshAdapter::OnReindexElements( const UEditableMesh* EditableMesh, const FElementIDRemappings& Remappings )
 {
-	RemapSparseArrayElements( RenderingPolygons, Remappings.NewPolygonIndexLookup );
-	RemapSparseArrayElements( RenderingPolygonGroups, Remappings.NewPolygonGroupIndexLookup );
+	RenderingPolygons.Remap( Remappings.NewPolygonIndexLookup );
+	RenderingPolygonGroups.Remap( Remappings.NewPolygonGroupIndexLookup );
 
 	// Always compact the rendering triangles
-	for( TSparseArray<FRenderingPolygonGroup>::TIterator It( RenderingPolygonGroups ); It; ++It )
+	for( const FPolygonGroupID PolygonGroupID : RenderingPolygonGroups.GetElementIDs() )
 	{
-		FRenderingPolygonGroup& RenderingPolygonGroup = *It;
-		const FMeshPolygonGroup& PolygonGroup = EditableMesh->PolygonGroups[ It.GetIndex() ];
+		FRenderingPolygonGroup& RenderingPolygonGroup = RenderingPolygonGroups[ PolygonGroupID ];
+		const FMeshPolygonGroup& PolygonGroup = EditableMesh->GetMeshDescription()->GetPolygonGroup( PolygonGroupID );
 
 		TSparseArray<FTriangleID> TriangleRemappings;
-		CompactSparseArrayElements( RenderingPolygonGroup.Triangles, TriangleRemappings );
+		RenderingPolygonGroup.Triangles.Compact( TriangleRemappings );
 
-		for( FMeshTriangle& Triangle : RenderingPolygonGroup.Triangles )
+		for( const FTriangleID TriangleID : RenderingPolygonGroup.Triangles.GetElementIDs() )
 		{
+			FMeshTriangle& Triangle = RenderingPolygonGroup.Triangles[ TriangleID ];
 			for( int32 TriangleVertexNumber = 0; TriangleVertexNumber < 3; ++TriangleVertexNumber )
 			{
 				const FVertexInstanceID OriginalVertexInstanceID = Triangle.GetVertexInstanceID( TriangleVertexNumber );
@@ -926,12 +932,14 @@ void UEditableStaticMeshAdapter::OnReindexElements( const UEditableMesh* Editabl
 			}
 		}
 
-		RenderingPolygonGroup.MaxTriangles = RenderingPolygonGroup.Triangles.GetMaxIndex();
+		RenderingPolygonGroup.MaxTriangles = RenderingPolygonGroup.Triangles.GetArraySize();
 
 		// Fix up references in referencing polygons
 		for( const FPolygonID PolygonID : PolygonGroup.Polygons )
 		{
-			FRenderingPolygon& RenderingPolygon = RenderingPolygons[ PolygonID.GetValue() ];
+			FRenderingPolygon& RenderingPolygon = RenderingPolygons[ PolygonID ];
+			RenderingPolygon.PolygonGroupID = PolygonGroupID;
+
 			for( FTriangleID& TriangleID : RenderingPolygon.TriangulatedPolygonTriangleIndices )
 			{
 				TriangleID = TriangleRemappings[ TriangleID.GetValue() ];
@@ -1062,12 +1070,14 @@ void UEditableStaticMeshAdapter::UpdateBounds( const UEditableMesh* EditableMesh
 			// 2) cache bounding box per polygon
 			// There are other cases where having polygon adjacency information (1) might be useful, so it's maybe worth considering.
 
-			for( const FMeshPolygon& Polygon : EditableMesh->Polygons )
+			const FVertexArray& Vertices = EditableMesh->GetMeshDescription()->Vertices();
+			const FVertexInstanceArray& VertexInstances = EditableMesh->GetMeshDescription()->VertexInstances();
+
+			for( const FVertexID VertexID : Vertices.GetElementIDs() )
 			{
-				for( const FVertexInstanceID VertexInstanceID : Polygon.PerimeterContour.VertexInstanceIDs )
+				if( !EditableMesh->IsOrphanedVertex( VertexID ) )
 				{
-					const FMeshVertexInstance& VertexInstance = EditableMesh->VertexInstances[ VertexInstanceID.GetValue() ];
-					BoundingBox += EditableMesh->Vertices[ VertexInstance.VertexID.GetValue() ].VertexPosition;
+					BoundingBox += Vertices[ VertexID ].VertexPosition;
 				}
 			}
 
@@ -1076,13 +1086,11 @@ void UEditableStaticMeshAdapter::UpdateBounds( const UEditableMesh* EditableMesh
 			// Calculate the bounding sphere, using the center of the bounding box as the origin.
 			BoundingBoxAndSphere.SphereRadius = 0.0f;
 
-			for( const FMeshPolygon& Polygon : EditableMesh->Polygons )
+			for( const FVertexID VertexID : Vertices.GetElementIDs() )
 			{
-				for( const FVertexInstanceID VertexInstanceID : Polygon.PerimeterContour.VertexInstanceIDs )
+				if( !EditableMesh->IsOrphanedVertex( VertexID ) )
 				{
-					const FMeshVertexInstance& VertexInstance = EditableMesh->VertexInstances[ VertexInstanceID.GetValue() ];
-					const FVector VertexPosition = EditableMesh->Vertices[ VertexInstance.VertexID.GetValue() ].VertexPosition;
-
+					const FVector VertexPosition = Vertices[ VertexID ].VertexPosition;
 					BoundingBoxAndSphere.SphereRadius = FMath::Max( ( VertexPosition - BoundingBoxAndSphere.Origin ).Size(), BoundingBoxAndSphere.SphereRadius );
 				}
 			}
@@ -1158,7 +1166,7 @@ void UEditableStaticMeshAdapter::UpdateCollision()
 
 void UEditableStaticMeshAdapter::OnSetVertexAttribute( const UEditableMesh* EditableMesh, const FVertexID VertexID, const FName AttributeName, const int32 AttributeIndex, const FVector4 AttributeValue )
 {
-	const FMeshVertex& Vertex = EditableMesh->Vertices[ VertexID.GetValue() ];
+	const FMeshVertex& Vertex = EditableMesh->GetMeshDescription()->GetVertex( VertexID );
 	FStaticMeshLODResources& StaticMeshLOD = GetStaticMeshLOD();
 
 	if( AttributeName == UEditableMeshAttribute::VertexPosition() )
@@ -1169,10 +1177,12 @@ void UEditableStaticMeshAdapter::OnSetVertexAttribute( const UEditableMesh* Edit
 		// (may also want different component eventually)
 		if( !EditableMesh->IsPreviewingSubdivisions() )
 		{
+			const FVertexInstanceArray& VertexInstances = EditableMesh->GetMeshDescription()->VertexInstances();
+
 			// Set the vertex buffer position of all of the vertex instances for this editable vertex
 			for( const FVertexInstanceID VertexInstanceID : Vertex.VertexInstanceIDs )
 			{
-				check( EditableMesh->VertexInstances.IsAllocated( VertexInstanceID.GetValue() ) );
+				check( VertexInstances.IsValid( VertexInstanceID ) );
 				StaticMeshLOD.PositionVertexBuffer.VertexPosition( VertexInstanceID.GetValue() ) = NewVertexPosition;
 			}
 		}
@@ -1209,7 +1219,7 @@ void UEditableStaticMeshAdapter::OnSetEdgeAttribute( const UEditableMesh* Editab
 
 void UEditableStaticMeshAdapter::OnSetVertexInstanceAttribute( const UEditableMesh* EditableMesh, const FVertexInstanceID VertexInstanceID, const FName AttributeName, const int32 AttributeIndex, const FVector4 AttributeValue )
 {
-	const FMeshVertexInstance& VertexInstance = EditableMesh->VertexInstances[ VertexInstanceID.GetValue() ];
+	const FMeshVertexInstance& VertexInstance = EditableMesh->GetMeshDescription()->GetVertexInstance( VertexInstanceID );
 	FStaticMeshLODResources& StaticMeshLOD = GetStaticMeshLOD();
 
 	if( AttributeName == UEditableMeshAttribute::VertexNormal() ||
@@ -1242,7 +1252,7 @@ void UEditableStaticMeshAdapter::OnSetVertexInstanceAttribute( const UEditableMe
 		{
 			const FColor NewColor = VertexInstance.Color.ToFColor( true );
 
-			if( StaticMeshLOD.ColorVertexBuffer.GetNumVertices() != EditableMesh->VertexInstances.GetMaxIndex() )
+			if( StaticMeshLOD.ColorVertexBuffer.GetNumVertices() != EditableMesh->GetMeshDescription()->VertexInstances().GetArraySize() )
 			{
 				if( VertexInstance.Color != FLinearColor::White )
 				{
@@ -1278,17 +1288,20 @@ void UEditableStaticMeshAdapter::OnCreateVertexInstances( const UEditableMesh* E
 		const int32 NumUVs = EditableMesh->GetTextureCoordinateCount();
 		const bool bHasColors = StaticMeshLOD.ColorVertexBuffer.GetNumVertices() > 0;
 
+		const FVertexArray& Vertices = EditableMesh->GetMeshDescription()->Vertices();
+		const FVertexInstanceArray& VertexInstances = EditableMesh->GetMeshDescription()->VertexInstances();
+
 		// Determine if we need to grow the render buffers
 		const int32 OldVertexBufferRenderingVertexCount = StaticMeshLOD.PositionVertexBuffer.GetNumVertices();
-		const int32 NumNewVertexBufferRenderingVertices = FMath::Max( 0, EditableMesh->VertexInstances.GetMaxIndex() - OldVertexBufferRenderingVertexCount );
+		const int32 NumNewVertexBufferRenderingVertices = FMath::Max( 0, VertexInstances.GetArraySize() - OldVertexBufferRenderingVertexCount );
 
 		static TArray<FStaticMeshBuildVertex> RenderingVerticesToAppend;
 		RenderingVerticesToAppend.SetNumUninitialized( NumNewVertexBufferRenderingVertices, false );
 
 		for( const FVertexInstanceID VertexInstanceID : VertexInstanceIDs )
 		{
-			const FMeshVertexInstance& VertexInstance = EditableMesh->VertexInstances[ VertexInstanceID.GetValue() ];
-			const FMeshVertex& ReferencedVertex = EditableMesh->Vertices[ VertexInstance.VertexID.GetValue() ];
+			const FMeshVertexInstance& VertexInstance = VertexInstances[ VertexInstanceID ];
+			const FMeshVertex& ReferencedVertex = Vertices[ VertexInstance.VertexID ];
 
 			const int32 NewRenderingVertexIndex = VertexInstanceID.GetValue();	// Rendering vertex indices are the same as vertex instance IDs
 
@@ -1359,7 +1372,8 @@ void UEditableStaticMeshAdapter::OnCreatePolygons( const UEditableMesh* Editable
 	// Add mirror polygons for static mesh adapter
 	for( const FPolygonID PolygonID : PolygonIDs )
 	{
-		RenderingPolygons.Insert( PolygonID.GetValue(), FRenderingPolygon() );
+		RenderingPolygons.Insert( PolygonID );
+		RenderingPolygons[ PolygonID ].PolygonGroupID = EditableMesh->GetGroupForPolygon( PolygonID );
 	}
 }
 
@@ -1372,15 +1386,18 @@ void UEditableStaticMeshAdapter::OnChangePolygonVertexInstances( const UEditable
 
 void UEditableStaticMeshAdapter::OnRetriangulatePolygons( const UEditableMesh* EditableMesh, const TArray<FPolygonID>& PolygonIDs )
 {
+	const FPolygonArray& Polygons = EditableMesh->GetMeshDescription()->Polygons();
+	const FPolygonGroupArray& PolygonGroups = EditableMesh->GetMeshDescription()->PolygonGroups();
+
 	for( const FPolygonID PolygonID : PolygonIDs )
 	{
-		const FPolygonGroupID PolygonGroupID = EditableMesh->GetGroupForPolygon( PolygonID );
+		const FMeshPolygon& Polygon = Polygons[ PolygonID ];
+		FRenderingPolygon& RenderingPolygon = RenderingPolygons[ PolygonID ];
 
-		const FMeshPolygonGroup& PolygonGroup = EditableMesh->PolygonGroups[ PolygonGroupID.GetValue() ];
-		const FMeshPolygon& Polygon = EditableMesh->Polygons[ PolygonID.GetValue() ];
+		const FPolygonGroupID PolygonGroupID = RenderingPolygon.PolygonGroupID;
 
-		FRenderingPolygonGroup& RenderingPolygonGroup = RenderingPolygonGroups[ PolygonGroupID.GetValue() ];
-		FRenderingPolygon& RenderingPolygon = RenderingPolygons[ PolygonID.GetValue() ];
+//		const FMeshPolygonGroup& PolygonGroup = EditableMesh->PolygonGroups[ PolygonGroupID.GetValue() ];
+		FRenderingPolygonGroup& RenderingPolygonGroup = RenderingPolygonGroups[ PolygonGroupID ];
 
 		const TArray<FMeshTriangle>& Triangles = Polygon.Triangles;
 
@@ -1397,7 +1414,7 @@ void UEditableStaticMeshAdapter::OnRetriangulatePolygons( const UEditableMesh* E
 				// See if the triangulation has changed even if the number of triangles is the same
 				for( int32 Index = 0; Index < Triangles.Num(); ++Index )
 				{
-					const FMeshTriangle& OldTriangle = RenderingPolygonGroup.Triangles[ RenderingPolygon.TriangulatedPolygonTriangleIndices[ Index ].GetValue() ];
+					const FMeshTriangle& OldTriangle = RenderingPolygonGroup.Triangles[ RenderingPolygon.TriangulatedPolygonTriangleIndices[ Index ] ];
 					const FMeshTriangle& NewTriangle = Triangles[ Index ];
 
 					if( OldTriangle.VertexInstanceID0 != NewTriangle.VertexInstanceID0 ||
@@ -1427,13 +1444,13 @@ void UEditableStaticMeshAdapter::OnRetriangulatePolygons( const UEditableMesh* E
 				const int32 NumNewTriangles = Triangles.Num();
 
 				// This is the number of entries currently unused in the Triangles sparse array
-				const int32 NumFreeTriangles = RenderingPolygonGroup.Triangles.GetMaxIndex() - RenderingPolygonGroup.Triangles.Num();
+				const int32 NumFreeTriangles = RenderingPolygonGroup.Triangles.GetArraySize() - RenderingPolygonGroup.Triangles.Num();
 
 				// This is the number of triangles we need to make extra space for (in both the sparse array and the index buffer)
 				const int32 NumTrianglesToReserve = FMath::Max( 0, NumNewTriangles - NumFreeTriangles );
 
 				// This is the number of triangles we will need to have allocated in the index buffer after adding the new triangles
-				const int32 NewTotalTriangles = RenderingPolygonGroup.Triangles.GetMaxIndex() + NumTrianglesToReserve;
+				const int32 NewTotalTriangles = RenderingPolygonGroup.Triangles.GetArraySize() + NumTrianglesToReserve;
 
 				// Reserve extra triangles if necessary.
 				if( NumTrianglesToReserve > 0 )
@@ -1446,16 +1463,16 @@ void UEditableStaticMeshAdapter::OnRetriangulatePolygons( const UEditableMesh* E
 				int32 MaxVertexIndex = RenderingSection.MaxVertexIndex;
 
 				// Create empty triangles for all of the new triangles we need, and keep track of their triangle indices
-				static TArray<int32> NewTriangleIndices;
+				static TArray<FTriangleID> NewTriangleIDs;
 				{
-					NewTriangleIndices.SetNumUninitialized( NumNewTriangles, false );
+					NewTriangleIDs.SetNumUninitialized( NumNewTriangles, false );
 
 					for( int32 TriangleToAddNumber = 0; TriangleToAddNumber < NumNewTriangles; ++TriangleToAddNumber )
 					{
-						const int32 NewTriangleIndex = RenderingPolygonGroup.Triangles.Add( FMeshTriangle() );
-						NewTriangleIndices[ TriangleToAddNumber ] = NewTriangleIndex;
+						const FTriangleID NewTriangleID = RenderingPolygonGroup.Triangles.Add();
+						NewTriangleIDs[ TriangleToAddNumber ] = NewTriangleID;
 
-						FMeshTriangle& NewTriangle = RenderingPolygonGroup.Triangles[ NewTriangleIndex ];
+						FMeshTriangle& NewTriangle = RenderingPolygonGroup.Triangles[ NewTriangleID ];
 						for( int32 TriangleVertexNumber = 0; TriangleVertexNumber < 3; ++TriangleVertexNumber )
 						{
 							const FVertexInstanceID VertexInstanceID = Triangles[ TriangleToAddNumber ].GetVertexInstanceID( TriangleVertexNumber );
@@ -1464,7 +1481,7 @@ void UEditableStaticMeshAdapter::OnRetriangulatePolygons( const UEditableMesh* E
 							MaxVertexIndex = FMath::Max( MaxVertexIndex, VertexInstanceID.GetValue() );
 						}
 
-						RenderingPolygon.TriangulatedPolygonTriangleIndices.Add( FTriangleID( NewTriangleIndex ) );
+						RenderingPolygon.TriangulatedPolygonTriangleIndices.Add( NewTriangleID );
 					}
 				}
 
@@ -1506,13 +1523,13 @@ void UEditableStaticMeshAdapter::OnRetriangulatePolygons( const UEditableMesh* E
 				{
 					for( int32 TriangleToAddNumber = 0; TriangleToAddNumber < NumNewTriangles; ++TriangleToAddNumber )
 					{
-						const int32 NewTriangleIndex = NewTriangleIndices[ TriangleToAddNumber ];
+						const FTriangleID NewTriangleID = NewTriangleIDs[ TriangleToAddNumber ];
 
 						for( int32 TriangleVertexNumber = 0; TriangleVertexNumber < 3; ++TriangleVertexNumber )
 						{
 							StaticMeshLOD.IndexBuffer.SetIndex(
-								FRenderingPolygonGroup::TriangleIndexToRenderingTriangleFirstIndex( RenderingSection, FTriangleID( NewTriangleIndex ) ) + TriangleVertexNumber,
-								RenderingPolygonGroup.Triangles[ NewTriangleIndex ].GetVertexInstanceID( TriangleVertexNumber ).GetValue() );
+								FRenderingPolygonGroup::TriangleIndexToRenderingTriangleFirstIndex( RenderingSection, NewTriangleID ) + TriangleVertexNumber,
+								RenderingPolygonGroup.Triangles[ NewTriangleID ].GetVertexInstanceID( TriangleVertexNumber ).GetValue() );
 						}
 					}
 
@@ -1562,17 +1579,17 @@ void UEditableStaticMeshAdapter::OnDeletePolygons( const UEditableMesh* Editable
 		DeletePolygonTriangles( EditableMesh, PolygonID );
 
 		// Delete the polygon from the static mesh adapter mirror
-		RenderingPolygons.RemoveAt( PolygonID.GetValue() );
+		RenderingPolygons.Remove( PolygonID );
 	}
 }
 
 
 void UEditableStaticMeshAdapter::DeletePolygonTriangles( const UEditableMesh* EditableMesh, const FPolygonID PolygonID )
 {
-	const FPolygonGroupID PolygonGroupID = EditableMesh->GetGroupForPolygon( PolygonID );
+	FRenderingPolygon& Polygon = RenderingPolygons[ PolygonID ];
+	const FPolygonGroupID PolygonGroupID = Polygon.PolygonGroupID;
 
-	FRenderingPolygonGroup& RenderingPolygonGroup = RenderingPolygonGroups[ PolygonGroupID.GetValue() ];
-	FRenderingPolygon& Polygon = RenderingPolygons[ PolygonID.GetValue() ];
+	FRenderingPolygonGroup& RenderingPolygonGroup = RenderingPolygonGroups[ PolygonGroupID ];
 
 	const int32 NumTrianglesToRemove = Polygon.TriangulatedPolygonTriangleIndices.Num();
 	if( NumTrianglesToRemove > 0 )
@@ -1592,9 +1609,9 @@ void UEditableStaticMeshAdapter::DeletePolygonTriangles( const UEditableMesh* Ed
 			const int32 MinVertexIndex = RenderingSection.MinVertexIndex;
 			const int32 MaxVertexIndex = RenderingSection.MaxVertexIndex;
 
-			for( const FTriangleID TriangleIndexToRemove : Polygon.TriangulatedPolygonTriangleIndices )
+			for( const FTriangleID TriangleIDToRemove : Polygon.TriangulatedPolygonTriangleIndices )
 			{
-				const FMeshTriangle& Triangle = RenderingPolygonGroup.Triangles[ TriangleIndexToRemove.GetValue() ];
+				const FMeshTriangle& Triangle = RenderingPolygonGroup.Triangles[ TriangleIDToRemove ];
 
 				for( int32 TriangleVertexNumber = 0; TriangleVertexNumber < 3; ++TriangleVertexNumber )
 				{
@@ -1619,7 +1636,7 @@ void UEditableStaticMeshAdapter::DeletePolygonTriangles( const UEditableMesh* Ed
 		for( const FTriangleID TriangleIndexToRemove : Polygon.TriangulatedPolygonTriangleIndices )
 		{
 			// Remove this triangle from our editable mesh
-			RenderingPolygonGroup.Triangles.RemoveAt( TriangleIndexToRemove.GetValue() );
+			RenderingPolygonGroup.Triangles.Remove( TriangleIndexToRemove );
 		}
 
 		if( !EditableMesh->IsPreviewingSubdivisions() )
@@ -1629,8 +1646,10 @@ void UEditableStaticMeshAdapter::DeletePolygonTriangles( const UEditableMesh* Ed
 				int32 MinVertexIndex = TNumericLimits<int32>::Max();
 				int32 MaxVertexIndex = TNumericLimits<int32>::Min();
 
-				for( const FMeshTriangle& Triangle : RenderingPolygonGroup.Triangles )
+				for( const FTriangleID TriangleID : RenderingPolygonGroup.Triangles.GetElementIDs() )
 				{
+					const FMeshTriangle& Triangle = RenderingPolygonGroup.Triangles[ TriangleID ];
+
 					for( int32 TriangleVertexNumber = 0; TriangleVertexNumber < 3; ++TriangleVertexNumber )
 					{
 						const FVertexInstanceID VertexInstanceID = Triangle.GetVertexInstanceID( TriangleVertexNumber );
@@ -1688,9 +1707,11 @@ FStaticMeshLODResources& UEditableStaticMeshAdapter::GetStaticMeshLOD()
 
 void UEditableStaticMeshAdapter::OnCreatePolygonGroups( const UEditableMesh* EditableMesh, const TArray<FPolygonGroupID>& PolygonGroupIDs )
 {
+	const FPolygonGroupArray& PolygonGroups = EditableMesh->GetMeshDescription()->PolygonGroups();
+
 	for( const FPolygonGroupID PolygonGroupID : PolygonGroupIDs )
 	{
-		const FMeshPolygonGroup& PolygonGroup = EditableMesh->PolygonGroups[ PolygonGroupID.GetValue() ];
+		const FMeshPolygonGroup& PolygonGroup = PolygonGroups[ PolygonGroupID ];
 
 		uint32 LODSectionIndex = 0;
 		if( !EditableMesh->IsPreviewingSubdivisions() )
@@ -1714,7 +1735,7 @@ void UEditableStaticMeshAdapter::OnCreatePolygonGroups( const UEditableMesh* Edi
 				FStaticMeshSection& PreviousStaticMeshSection = StaticMeshLOD.Sections[ LODSectionIndex - 1 ];
 				const FPolygonGroupID PreviousPolygonGroupID = GetSectionForRenderingSectionIndex( LODSectionIndex - 1 );
 				check( PreviousPolygonGroupID != FPolygonGroupID::Invalid );
-				StaticMeshSection.FirstIndex = PreviousStaticMeshSection.FirstIndex + RenderingPolygonGroups[ PreviousPolygonGroupID.GetValue() ].MaxTriangles * 3;
+				StaticMeshSection.FirstIndex = PreviousStaticMeshSection.FirstIndex + RenderingPolygonGroups[ PreviousPolygonGroupID ].MaxTriangles * 3;
 
 				// @todo mesheditor: if this check is valid, we can dispense with the above and just set StaticMeshSection.FirstIndex according to the current length of the idnex buffer.
 				check( StaticMeshLOD.IndexBuffer.GetNumIndices() == StaticMeshSection.FirstIndex );
@@ -1727,14 +1748,13 @@ void UEditableStaticMeshAdapter::OnCreatePolygonGroups( const UEditableMesh* Edi
 			StaticMeshSection.bEnableCollision = PolygonGroup.bEnableCollision;
 			StaticMeshSection.bCastShadow = PolygonGroup.bCastShadow;
 
-			check( GetStaticMeshMaterialIndex( StaticMesh, PolygonGroup.Material ) == INDEX_NONE );
-			const int32 MaterialIndex = StaticMesh->StaticMaterials.Add( FStaticMaterial( PolygonGroup.Material ) );
-			StaticMeshSection.MaterialIndex = MaterialIndex;
+			UMaterialInterface* Material = Cast<UMaterialInterface>( PolygonGroup.MaterialAsset.ResolveObject() );
+			StaticMeshSection.MaterialIndex = StaticMesh->StaticMaterials.Emplace( Material, PolygonGroup.MaterialSlotName, PolygonGroup.ImportedMaterialSlotName );
 		}
 
 		// Insert the rendering polygon group for keeping track of these index buffer properties
-		RenderingPolygonGroups.Insert( PolygonGroupID.GetValue(), FRenderingPolygonGroup() );
-		FRenderingPolygonGroup& RenderingPolygonGroup = RenderingPolygonGroups[ PolygonGroupID.GetValue() ];
+		RenderingPolygonGroups.Insert( PolygonGroupID );
+		FRenderingPolygonGroup& RenderingPolygonGroup = RenderingPolygonGroups[ PolygonGroupID ];
 
 		RenderingPolygonGroup.RenderingSectionIndex = LODSectionIndex;
 		RenderingPolygonGroup.MaxTriangles = 0;
@@ -1744,21 +1764,26 @@ void UEditableStaticMeshAdapter::OnCreatePolygonGroups( const UEditableMesh* Edi
 
 void UEditableStaticMeshAdapter::OnDeletePolygonGroups( const UEditableMesh* EditableMesh, const TArray<FPolygonGroupID>& PolygonGroupIDs )
 {
+	const FPolygonGroupArray& PolygonGroups = EditableMesh->GetMeshDescription()->PolygonGroups();
+
 	for( const FPolygonGroupID PolygonGroupID : PolygonGroupIDs )
 	{
-		const FMeshPolygonGroup& PolygonGroup = EditableMesh->PolygonGroups[ PolygonGroupID.GetValue() ];
-		FRenderingPolygonGroup& RenderingPolygonGroup = RenderingPolygonGroups[ PolygonGroupID.GetValue() ];
+		const FMeshPolygonGroup& PolygonGroup = PolygonGroups[ PolygonGroupID ];
+		FRenderingPolygonGroup& RenderingPolygonGroup = RenderingPolygonGroups[ PolygonGroupID ];
 
 		// Remove material slot associated with section
-		// @todo mesheditor: can more than one section share a material? Mesh editor currently assumes not, but this will break anything which does
-		const int32 MaterialIndex = GetStaticMeshMaterialIndex( StaticMesh, PolygonGroup.Material );
+		const int32 MaterialIndex = StaticMesh->StaticMaterials.IndexOfByPredicate(
+			[ &PolygonGroup ]( const FStaticMaterial& StaticMaterial ) { return StaticMaterial.MaterialSlotName == PolygonGroup.MaterialSlotName; }
+		);
 		StaticMesh->StaticMaterials.RemoveAt( MaterialIndex );
 
 		// Adjust rendering indices held by sections: any index above the one we just deleted now needs to be decremented.
 		const uint32 RenderingSectionIndex = RenderingPolygonGroup.RenderingSectionIndex;
 
-		for( FRenderingPolygonGroup& PolygonGroupToAdjust : RenderingPolygonGroups )
+		for( const FPolygonGroupID PolygonGroupIDToAdjust : RenderingPolygonGroups.GetElementIDs() )
 		{
+			FRenderingPolygonGroup& PolygonGroupToAdjust = RenderingPolygonGroups[ PolygonGroupIDToAdjust ];
+
 			if( PolygonGroupToAdjust.RenderingSectionIndex > RenderingSectionIndex )
 			{
 				PolygonGroupToAdjust.RenderingSectionIndex--;
@@ -1799,18 +1824,41 @@ void UEditableStaticMeshAdapter::OnDeletePolygonGroups( const UEditableMesh* Edi
 		}
 
 		// Remove the rendering polygon group from the sparse array
-		RenderingPolygonGroups.RemoveAt( PolygonGroupID.GetValue() );
+		RenderingPolygonGroups.Remove( PolygonGroupID );
 	}
 }
 
 
+void UEditableStaticMeshAdapter::OnAssignPolygonsToPolygonGroups( const UEditableMesh* EditableMesh, const TArray<FPolygonGroupForPolygon>& PolygonGroupForPolygons )
+{
+	static TArray<FPolygonID> PolygonIDsToRebuild;
+	PolygonIDsToRebuild.Reset( PolygonGroupForPolygons.Num() );
+
+	for( const FPolygonGroupForPolygon& PolygonGroupForPolygon : PolygonGroupForPolygons )
+	{
+		const FPolygonID PolygonID = PolygonGroupForPolygon.PolygonID;
+		const FPolygonGroupID NewPolygonGroupID = PolygonGroupForPolygon.PolygonGroupID;
+
+		DeletePolygonTriangles( EditableMesh, PolygonID );
+
+		FRenderingPolygon& RenderingPolygon = RenderingPolygons[ PolygonID ];
+		RenderingPolygon.PolygonGroupID = NewPolygonGroupID;
+
+		PolygonIDsToRebuild.Add( PolygonID );
+	}
+
+	OnRetriangulatePolygons( EditableMesh, PolygonIDsToRebuild );
+}
+
+
 FPolygonGroupID UEditableStaticMeshAdapter::GetSectionForRenderingSectionIndex( const int32 RenderingSectionIndex ) const
 {
-	for( auto It = RenderingPolygonGroups.CreateConstIterator(); It; ++It )
+	for( const FPolygonGroupID PolygonGroupID : RenderingPolygonGroups.GetElementIDs() )
 	{
-		if( It->RenderingSectionIndex == RenderingSectionIndex )
+		const FRenderingPolygonGroup& RenderingPolygonGroup = RenderingPolygonGroups[ PolygonGroupID ];
+		if( RenderingPolygonGroup.RenderingSectionIndex == RenderingSectionIndex )
 		{
-			return FPolygonGroupID( It.GetIndex() );
+			return PolygonGroupID;
 		}
 	}
 
diff --git a/Engine/Plugins/Runtime/EditableMesh/Source/EditableMesh/EditableStaticMeshAdapter.h b/Engine/Plugins/Runtime/EditableMesh/Source/EditableMesh/EditableStaticMeshAdapter.h
--- a/Engine/Plugins/Runtime/EditableMesh/Source/EditableMesh/EditableStaticMeshAdapter.h
+++ b/Engine/Plugins/Runtime/EditableMesh/Source/EditableMesh/EditableStaticMeshAdapter.h
@@ -43,6 +43,10 @@ struct FRenderingPolygon
 {
 	GENERATED_BODY()
 
+	/** Which rendering polygon group the polygon is in */
+	UPROPERTY()
+	FPolygonGroupID PolygonGroupID;
+
 	/** This is a list of indices of triangles in the FRenderingPolygonGroup::Triangles array.
 	    We use this to maintain a record of which triangles in the section belong to this polygon. */
 	UPROPERTY()
@@ -51,6 +55,7 @@ struct FRenderingPolygon
 	/** Serializer */
 	friend FArchive& operator<<( FArchive& Ar, FRenderingPolygon& Polygon )
 	{
+		Ar << Polygon.PolygonGroupID;
 		Ar << Polygon.TriangulatedPolygonTriangleIndices;
 		return Ar;
 	}
@@ -72,7 +77,7 @@ struct FRenderingPolygonGroup
 
 	/** Sparse array of triangles, that matches the triangles in the mesh index buffers.  Elements that
 	    aren't allocated will be stored as degenerates in the mesh index buffer. */
-	TSparseArray<FMeshTriangle> Triangles;
+	TMeshElementArray<FMeshTriangle, FTriangleID> Triangles;
 
 
 	/** Converts from an index in our Triangles array to an index of a rendering triangle's first vertex in the rendering mesh's index buffer */
@@ -92,7 +97,7 @@ struct FRenderingPolygonGroup
 	{
 		Ar << Section.RenderingSectionIndex;
 		Ar << Section.MaxTriangles;	// @todo mesheditor serialization: Should not need to be serialized if we triangulate after load
-		SerializeSparseArray( Ar, Section.Triangles );	// @todo mesheditor serialization: Should not need to be serialized if we triangulate after load
+		Ar << Section.Triangles;
 		return Ar;
 	}
 };
@@ -147,6 +152,7 @@ public:
 	virtual void OnChangePolygonVertexInstances( const UEditableMesh* EditableMesh, const TArray<FPolygonID>& PolygonIDs ) override;
 	virtual void OnCreatePolygonGroups( const UEditableMesh* EditableMesh, const TArray<FPolygonGroupID>& PolygonGroupIDs ) override;
 	virtual void OnDeletePolygonGroups( const UEditableMesh* EditableMesh, const TArray<FPolygonGroupID>& PolygonGroupIDs ) override;
+	virtual void OnAssignPolygonsToPolygonGroups( const UEditableMesh* EditableMesh, const TArray<FPolygonGroupForPolygon>& PolygonGroupForPolygons ) override;
 	virtual void OnRetriangulatePolygons( const UEditableMesh* EditableMesh, const TArray<FPolygonID>& PolygonIDs ) override;
 
 
@@ -192,13 +198,13 @@ private:
 	static const uint32 IndexBufferInterSectionGap = 32;
 
 	/** All of the polygons in this mesh */
-	TSparseArray< FRenderingPolygon > RenderingPolygons;
+	TMeshElementArray<FRenderingPolygon, FPolygonID> RenderingPolygons;
 
 	/** All of the polygon groups in this mesh */
-	TSparseArray< FRenderingPolygonGroup > RenderingPolygonGroups;
+	TMeshElementArray<FRenderingPolygonGroup, FPolygonGroupID> RenderingPolygonGroups;
 
 	/** Used to refresh all components in the scene that may be using a mesh we're editing */
-	TSharedPtr< class FStaticMeshComponentRecreateRenderStateContext > RecreateRenderStateContext;
+	TSharedPtr<class FStaticMeshComponentRecreateRenderStateContext> RecreateRenderStateContext;
 
 	/** Cached bounding box for the static mesh.  This bounds can be (temporarily) larger than the actual mesh itself as
 	    an optimization. */
diff --git a/Engine/Plugins/Runtime/EditableMesh/Source/EditableMesh/Public/EditableMesh.h b/Engine/Plugins/Runtime/EditableMesh/Source/EditableMesh/Public/EditableMesh.h
--- a/Engine/Plugins/Runtime/EditableMesh/Source/EditableMesh/Public/EditableMesh.h
+++ b/Engine/Plugins/Runtime/EditableMesh/Source/EditableMesh/Public/EditableMesh.h
@@ -4,6 +4,7 @@
 
 #include "CoreMinimal.h"
 #include "UObject/Object.h"
+#include "MeshDescription.h"
 #include "EditableMeshTypes.h"
 #include "EditableMeshCustomVersion.h"
 #include "Change.h"		// For TUniquePtr<FChange>
@@ -138,266 +139,6 @@ void RemapSparseArrayElements( TSparseArray<T>& Array, const TSparseArray<Elemen
 }
 
 
-USTRUCT()
-struct FMeshVertex
-{
-	GENERATED_BODY()
-
-	/** Position of the vertex. */
-	UPROPERTY()
-	FVector VertexPosition;
-
-	/** All of vertex instances which reference this vertex (for split vertex support) */
-	UPROPERTY()
-	TArray<FVertexInstanceID> VertexInstanceIDs;
-
-	/** The edges connected to this vertex */
-	UPROPERTY()
-	TArray<FEdgeID> ConnectedEdgeIDs;
-
-	/** When subdivisions are enabled, this controls how sharp the vertex is, between 0.0 and 1.0. */
-	UPROPERTY()
-	float CornerSharpness;	// @todo mesheditor subdiv: Not really used by static meshes at all.  Only for UEditableMeshes that use subdivision features.  Move elsewhere.?
-
-	/** Serializer */
-	friend FArchive& operator<<( FArchive& Ar, FMeshVertex& Vertex )
-	{
-		Ar << Vertex.VertexPosition;
-		Ar << Vertex.VertexInstanceIDs;
-		Ar << Vertex.ConnectedEdgeIDs;
-		Ar << Vertex.CornerSharpness;
-		return Ar;
-	}
-};
-
-
-USTRUCT()
-struct FMeshVertexInstance
-{
-	GENERATED_BODY()
-
-	/** The vertex this is instancing */
-	UPROPERTY()
-	FVertexID VertexID;
-
-	/** List of connected polygons */
-	UPROPERTY()
-	TArray<FPolygonID> ConnectedPolygons;
-
-	/** UVs for the vertex instance */
-	UPROPERTY()
-	TArray<FVector2D> VertexUVs;
-
-	/** Normal vector */
-	UPROPERTY()
-	FVector Normal;
-	
-	/** Tangent vector */
-	UPROPERTY()
-	FVector Tangent;
-
-	/** Basis determinant sign used to calculate the sense of the binormal */
-	UPROPERTY()
-	float BinormalSign;
-
-	/** Vertex color */
-	UPROPERTY()
-	FLinearColor Color;
-
-	/** Serializer */
-	friend FArchive& operator<<( FArchive& Ar, FMeshVertexInstance& VertexInstance )
-	{
-		Ar << VertexInstance.VertexID;
-		Ar << VertexInstance.ConnectedPolygons;
-		Ar << VertexInstance.VertexUVs;
-		Ar << VertexInstance.Normal;
-		Ar << VertexInstance.Tangent;
-		Ar << VertexInstance.BinormalSign;
-		Ar << VertexInstance.Color;
-		return Ar;
-	}
-};
-
-
-USTRUCT()
-struct FMeshEdge
-{
-	GENERATED_BODY()
-
-	/** IDs of the two editable mesh vertices that make up this edge.  The winding direction is not defined. */
-	UPROPERTY()
-	FVertexID VertexIDs[ 2 ];
-
-	/** The polygons that share this edge.  It's best if there are always only two polygons that share
-	    the edge, and those polygons are facing the same direction */
-	UPROPERTY()
-	TArray<FPolygonID> ConnectedPolygons;
-
-	/** Whether this edge is 'hard' or not, for the purpose of vertex normal and tangent generation */
-	UPROPERTY()
-	bool bIsHardEdge;
-
-	/** When subdivisions are enabled, this controls how sharp the creasing of this edge will be, between 0.0 and 1.0. */
-	UPROPERTY()
-	float CreaseSharpness;
-
-	/** Serializer */
-	friend FArchive& operator<<( FArchive& Ar, FMeshEdge& Edge )
-	{
-		Ar << Edge.VertexIDs[ 0 ];
-		Ar << Edge.VertexIDs[ 1 ];
-		Ar << Edge.ConnectedPolygons;
-		Ar << Edge.bIsHardEdge;
-		Ar << Edge.CreaseSharpness;
-		return Ar;
-	}
-};
-
-
-USTRUCT()
-struct FMeshPolygonContour
-{
-	GENERATED_BODY()
-
-	/** The ordered list of vertex instances which make up the polygon contour. The winding direction is counter-clockwise. */
-	UPROPERTY()
-	TArray<FVertexInstanceID> VertexInstanceIDs;
-
-	/** Serializer */
-	friend FArchive& operator<<( FArchive& Ar, FMeshPolygonContour& Contour )
-	{
-		Ar << Contour.VertexInstanceIDs;
-		return Ar;
-	}
-};
-
-
-USTRUCT( BlueprintType )
-struct FMeshTriangle
-{
-	GENERATED_BODY()
-
-	/** First vertex instance that makes up this triangle.  Indices must be ordered counter-clockwise. */
-	UPROPERTY( BlueprintReadWrite, Category="Editable Mesh" )
-	FVertexInstanceID VertexInstanceID0;
-
-	/** Second vertex instance that makes up this triangle.  Indices must be ordered counter-clockwise. */
-	UPROPERTY( BlueprintReadWrite, Category="Editable Mesh" )
-	FVertexInstanceID VertexInstanceID1;
-
-	/** Third vertex instance that makes up this triangle.  Indices must be ordered counter-clockwise. */
-	UPROPERTY( BlueprintReadWrite, Category="Editable Mesh" )
-	FVertexInstanceID VertexInstanceID2;
-
-	/** Gets the specified triangle vertex instance ID.  Pass an index between 0 and 2 inclusive. */
-	inline FVertexInstanceID GetVertexInstanceID( const int32 Index ) const
-	{
-		checkSlow( Index >= 0 && Index <= 2 );
-		return reinterpret_cast<const FVertexInstanceID*>( this )[ Index ];
-	}
-
-	/** Sets the specified triangle vertex instance ID.  Pass an index between 0 and 2 inclusive, and the new vertex instance ID to store. */
-	inline void SetVertexInstanceID( const int32 Index, const FVertexInstanceID NewVertexInstanceID )
-	{
-		checkSlow( Index >= 0 && Index <= 2 );
-		( reinterpret_cast<FVertexInstanceID*>( this )[ Index ] ) = NewVertexInstanceID;
-	}
-
-	/** Serializer */
-	friend FArchive& operator<<( FArchive& Ar, FMeshTriangle& Tri )
-	{
-		Ar << Tri.VertexInstanceID0;
-		Ar << Tri.VertexInstanceID1;
-		Ar << Tri.VertexInstanceID2;
-		return Ar;
-	}
-};
-
-
-USTRUCT()
-struct FMeshPolygon
-{
-	GENERATED_BODY()
-
-	/** The outer boundary edges of this polygon */
-	UPROPERTY()
-	FMeshPolygonContour PerimeterContour;
-
-	/** Optional inner contours of this polygon that define holes inside of the polygon.  For the geometry to
-	    be considered valid, the hole contours should reside within the boundary of the polygon perimeter contour, 
-		and must not overlap each other.  No "nesting" of polygons inside the holes is supported -- those are 
-		simply separate polygons */
-	UPROPERTY()
-	TArray<FMeshPolygonContour> HoleContours;
-
-	/** List of triangles which make up this polygon */
-	UPROPERTY()
-	TArray<FMeshTriangle> Triangles;
-
-	/** The polygon group which contains this polygon */
-	UPROPERTY()
-	FPolygonGroupID PolygonGroupID;
-
-	/** Cached normal */
-	UPROPERTY()
-	FVector PolygonNormal;
-
-	/** Cached tangent */
-	UPROPERTY()
-	FVector PolygonTangent;
-
-	/** Cached binormal */
-	UPROPERTY()
-	FVector PolygonBinormal;
-
-	/** Cached center */
-	UPROPERTY()
-	FVector PolygonCenter;
-
-	/** Serializer */
-	friend FArchive& operator<<( FArchive& Ar, FMeshPolygon& Polygon )
-	{
-		Ar << Polygon.PerimeterContour;
-		Ar << Polygon.HoleContours;
-		Ar << Polygon.Triangles;
-		Ar << Polygon.PolygonGroupID;
-		return Ar;
-	}
-};
-
-
-USTRUCT()
-struct FMeshPolygonGroup
-{
-	GENERATED_BODY()
-
-	/** The material for this mesh section */
-	UPROPERTY()
-	UMaterialInterface* Material;
-
-	/** If true, collision is enabled for this section. */
-	UPROPERTY()
-	bool bEnableCollision;
-
-	/** If true, this section will cast a shadow */
-	UPROPERTY()
-	bool bCastShadow;
-
-	/** All polygons in this group */
-	TArray<FPolygonID> Polygons;
-
-	/** Serializer */
-	friend FArchive& operator<<( FArchive& Ar, FMeshPolygonGroup& PolygonGroup )
-	{
-		Ar << PolygonGroup.Material;
-		Ar << PolygonGroup.bEnableCollision;
-		Ar << PolygonGroup.bCastShadow;
-		Ar << PolygonGroup.Polygons;
-		return Ar;
-	}
-};
-
-
 UENUM( BlueprintType )
 enum class EInsetPolygonsMode : uint8
 {
@@ -418,46 +159,6 @@ enum class ETriangleTessellationMode : uint8
 };
 
 
-struct FElementIDRemappings
-{
-	TSparseArray<FVertexID> NewVertexIndexLookup;
-	TSparseArray<FVertexInstanceID> NewVertexInstanceIndexLookup;
-	TSparseArray<FEdgeID> NewEdgeIndexLookup;
-	TSparseArray<FPolygonID> NewPolygonIndexLookup;
-	TSparseArray<FPolygonGroupID> NewPolygonGroupIndexLookup;
-
-	FVertexID GetRemappedVertexID( FVertexID VertexID ) const
-	{
-		checkSlow( NewVertexIndexLookup.IsAllocated( VertexID.GetValue() ) );
-		return NewVertexIndexLookup[ VertexID.GetValue() ];
-	}
-
-	FVertexInstanceID GetRemappedVertexInstanceID( FVertexInstanceID VertexInstanceID ) const
-	{
-		checkSlow( NewVertexInstanceIndexLookup.IsAllocated( VertexInstanceID.GetValue() ) );
-		return NewVertexInstanceIndexLookup[ VertexInstanceID.GetValue() ];
-	}
-
-	FEdgeID GetRemappedEdgeID( FEdgeID EdgeID ) const
-	{
-		checkSlow( NewEdgeIndexLookup.IsAllocated( EdgeID.GetValue() ) );
-		return NewEdgeIndexLookup[ EdgeID.GetValue() ];
-	}
-
-	FPolygonID GetRemappedPolygonID( FPolygonID PolygonID ) const
-	{
-		checkSlow( NewPolygonIndexLookup.IsAllocated( PolygonID.GetValue() ) );
-		return NewPolygonIndexLookup[ PolygonID.GetValue() ];
-	}
-
-	FPolygonGroupID GetRemappedPolygonGroupID( FPolygonGroupID PolygonGroupID ) const
-	{
-		check( NewPolygonGroupIndexLookup.IsAllocated( PolygonGroupID.GetValue() ) );
-		return NewPolygonGroupIndexLookup[ PolygonGroupID.GetValue() ];
-	}
-};
-
-
 UCLASS( BlueprintType )
 class EDITABLEMESH_API UEditableMesh : public UObject
 {
@@ -478,6 +179,8 @@ public:
 	/** Remaps mesh element arrays according to the provided remappings, in order to undo a compact operation */
 	void Uncompact( const FElementIDRemappings& Remappings );
 
+	UMeshDescription* GetMeshDescription() const { return MeshDescription; }
+
 	UFUNCTION( BlueprintCallable, Category="Editable Mesh" ) void InitializeAdapters();
 	UFUNCTION( BlueprintCallable, Category="Editable Mesh" ) void RebuildRenderMesh();
 	UFUNCTION( BlueprintCallable, Category="Editable Mesh" ) void StartModification( const EMeshModificationType MeshModificationType, const EMeshTopologyChange MeshTopologyChange );
@@ -490,57 +193,150 @@ public:
 	UFUNCTION( BlueprintCallable, Category="Editable Mesh" ) UEditableMesh* RevertInstance();
 	UFUNCTION( BlueprintCallable, Category="Editable Mesh" ) void PropagateInstanceChanges();
 
-	UFUNCTION( BlueprintPure, Category="Editable Mesh" ) int32 GetVertexCount() const;
-	UFUNCTION( BlueprintPure, Category="Editable Mesh" ) int32 GetVertexArraySize() const;
-	UFUNCTION( BlueprintPure, Category="Editable Mesh" ) bool IsValidVertex( const FVertexID VertexID ) const;
 	UFUNCTION( BlueprintPure, Category="Editable Mesh" ) FVector4 GetVertexAttribute( const FVertexID VertexID, const FName AttributeName, const int32 AttributeIndex ) const;
-	UFUNCTION( BlueprintPure, Category="Editable Mesh" ) int32 GetVertexConnectedEdgeCount( const FVertexID VertexID ) const;
-	UFUNCTION( BlueprintPure, Category="Editable Mesh" ) FEdgeID GetVertexConnectedEdge( const FVertexID VertexID, const int32 ConnectedEdgeNumber ) const;
-
-	UFUNCTION( BlueprintPure, Category="Editable Mesh" ) int32 GetVertexInstanceCount() const;
-	UFUNCTION( BlueprintPure, Category="Editable Mesh" ) int32 GetVertexInstanceArraySize() const;
 	UFUNCTION( BlueprintPure, Category="Editable Mesh" ) FVector4 GetVertexInstanceAttribute( const FVertexInstanceID VertexInstanceID, const FName AttributeName, const int32 AttributeIndex ) const;
-
 	UFUNCTION( BlueprintPure, Category="Editable Mesh" ) FVector4 GetEdgeAttribute( const FEdgeID EdgeID, const FName AttributeName, const int32 AttributeIndex ) const;
-	UFUNCTION( BlueprintPure, Category="Editable Mesh" ) int32 GetEdgeCount() const;
-	UFUNCTION( BlueprintPure, Category="Editable Mesh" ) int32 GetEdgeArraySize() const;
-	UFUNCTION( BlueprintPure, Category="Editable Mesh" ) bool IsValidEdge( const FEdgeID EdgeID ) const;
-	UFUNCTION( BlueprintPure, Category="Editable Mesh" ) FVertexID GetEdgeVertex( const FEdgeID EdgeID, const int32 EdgeVertexNumber ) const;
-	UFUNCTION( BlueprintPure, Category="Editable Mesh" ) int32 GetEdgeConnectedPolygonCount( const FEdgeID EdgeID ) const;
-	UFUNCTION( BlueprintPure, Category="Editable Mesh" ) FPolygonID GetEdgeConnectedPolygon( const FEdgeID EdgeID, const int32 ConnectedPolygonNumber ) const;
-
-	UFUNCTION( BlueprintPure, Category="Editable Mesh" ) int32 GetPolygonGroupCount() const;
-	UFUNCTION( BlueprintPure, Category="Editable Mesh" ) int32 GetPolygonGroupArraySize() const;
-	UFUNCTION( BlueprintPure, Category="Editable Mesh" ) bool IsValidPolygonGroup( const FPolygonGroupID PolygonGroupID ) const;
-	UFUNCTION( BlueprintPure, Category="Editable Mesh" ) int32 GetPolygonCountInGroup( const FPolygonGroupID PolygonGroupID ) const;
-	UFUNCTION( BlueprintPure, Category="Editable Mesh" ) FPolygonID GetPolygonInGroup( const FPolygonGroupID PolygonGroupID, const int32 PolygonNumber ) const;
-
-	UFUNCTION( BlueprintPure, Category="Editable Mesh" ) int32 GetPolygonCount() const;
-	UFUNCTION( BlueprintPure, Category="Editable Mesh" ) int32 GetPolygonArraySize() const;
-	UFUNCTION( BlueprintPure, Category="Editable Mesh" ) bool IsValidPolygon( const FPolygonID PolygonID ) const;
-	UFUNCTION( BlueprintPure, Category="Editable Mesh" ) FPolygonGroupID GetGroupForPolygon( const FPolygonID PolygonID ) const;
-
-//	UFUNCTION( BlueprintPure, Category="Editable Mesh" ) virtual int32 GetTriangleCount( const FPolygonGroupID PolygonGroupID ) const PURE_VIRTUAL(,return 0;);
-//	UFUNCTION( BlueprintPure, Category="Editable Mesh" ) virtual int32 GetTriangleArraySize( const FPolygonGroupID PolygonGroupID ) const PURE_VIRTUAL(,return 0;);
-
-	UFUNCTION( BlueprintPure, Category="Editable Mesh" ) int32 GetPolygonPerimeterVertexCount( const FPolygonID PolygonID ) const;
-	UFUNCTION( BlueprintPure, Category="Editable Mesh" ) FVertexID GetPolygonPerimeterVertex( const FPolygonID PolygonID, const int32 PolygonVertexNumber ) const;
-	UFUNCTION( BlueprintPure, Category="Editable Mesh" ) FVertexInstanceID GetPolygonPerimeterVertexInstance( const FPolygonID PolygonID, const int32 PolygonVertexNumber ) const;
-	UFUNCTION( BlueprintPure, Category="Editable Mesh" ) FVector4 GetPolygonPerimeterVertexAttribute( const FPolygonID PolygonID, const int32 PolygonVertexNumber, const FName AttributeName, const int32 AttributeIndex ) const;
-	UFUNCTION( BlueprintPure, Category="Editable Mesh" ) int32 GetPolygonHoleCount( const FPolygonID PolygonID ) const;
-	UFUNCTION( BlueprintPure, Category="Editable Mesh" ) int32 GetPolygonHoleVertexCount( const FPolygonID PolygonID, const int32 HoleNumber ) const;
-	UFUNCTION( BlueprintPure, Category="Editable Mesh" ) FVertexID GetPolygonHoleVertex( const FPolygonID PolygonID, const int32 HoleNumber, const int32 PolygonVertexNumber ) const;
-	UFUNCTION( BlueprintPure, Category="Editable Mesh" ) FVertexInstanceID GetPolygonHoleVertexInstance( const FPolygonID PolygonID, const int32 HoleNumber, const int32 PolygonVertexNumber ) const;
-	UFUNCTION( BlueprintPure, Category="Editable Mesh" ) FVector4 GetPolygonHoleVertexAttribute( const FPolygonID PolygonID, const int32 HoleNumber, const int32 PolygonVertexNumber, const FName AttributeName, const int32 AttributeIndex ) const;
-
-	UFUNCTION( BlueprintPure, Category="Editable Mesh" ) int32 GetPolygonTriangulatedTriangleCount( const FPolygonID PolygonID ) const;
-	UFUNCTION( BlueprintPure, Category="Editable Mesh" ) FMeshTriangle GetPolygonTriangulatedTriangle( const FPolygonID PolygonID, int32 PolygonTriangleNumber ) const;
-	UFUNCTION( BlueprintPure, Category="Editable Mesh" ) FVector GetPolygonTriangulatedTriangleVertexPosition( const FPolygonID PolygonID, const int32 PolygonTriangleNumber, const int32 TriangleVertexNumber ) const;
+
+	/** Returns the number of vertices in this mesh */
+	UFUNCTION( BlueprintPure, Category="Editable Mesh" )
+	int32 GetVertexCount() const;
+
+	/** Returns whether the given vertex ID is valid */
+	UFUNCTION( BlueprintPure, Category="Editable Mesh" )
+	bool IsValidVertex( const FVertexID VertexID ) const;
+
+	/** Returns whether the given vertex ID is orphaned */
+	UFUNCTION( BlueprintPure, Category="Editable Mesh" )
+	bool IsOrphanedVertex( const FVertexID VertexID ) const;
+
+	/** Returns the number of edges connected to this vertex */
+	UFUNCTION( BlueprintPure, Category="Editable Mesh" )
+	int32 GetVertexConnectedEdgeCount( const FVertexID VertexID ) const;
+
+	/** Returns the requested edge connected to this vertex */
+	UFUNCTION( BlueprintPure, Category="Editable Mesh" )
+	FEdgeID GetVertexConnectedEdge( const FVertexID VertexID, const int32 ConnectedEdgeNumber ) const;
+
+	/** Returns the number of vertex instances in this mesh */
+	UFUNCTION( BlueprintPure, Category="Editable Mesh" )
+	int32 GetVertexInstanceCount() const;
+
+	/** Returns the vertex ID which the given vertex instance is instancing */
+	UFUNCTION( BlueprintPure, Category="Editable Mesh" )
+	FVertexID GetVertexInstanceVertex( const FVertexInstanceID VertexInstanceID ) const;
+
+	/** Returns the number of polygons connected to this vertex instance */
+	UFUNCTION( BlueprintPure, Category="Editable Mesh" )
+	int32 GetVertexInstanceConnectedPolygonCount( const FVertexInstanceID VertexInstanceID ) const;
+
+	/** Returns the indexed polygon connected to this vertex instance */
+	UFUNCTION( BlueprintPure, Category="Editable Mesh" )
+	FPolygonID GetVertexInstanceConnectedPolygon( const FVertexInstanceID VertexInstanceID, const int32 ConnectedPolygonNumber ) const;
+
+	/** Returns the number of edges in this mesh */
+	UFUNCTION( BlueprintPure, Category="Editable Mesh" )
+	int32 GetEdgeCount() const;
+
+	/** Returns whether the given edge ID is valid */
+	UFUNCTION( BlueprintPure, Category="Editable Mesh" )
+	bool IsValidEdge( const FEdgeID EdgeID ) const;
+
+	/** Returns the given indexed vertex for this edge. EdgeVertexNumber must be 0 or 1. */
+	UFUNCTION( BlueprintPure, Category="Editable Mesh" )
+	FVertexID GetEdgeVertex( const FEdgeID EdgeID, const int32 EdgeVertexNumber ) const;
+
+	/** Returns the number of polygons connected to this edge */
+	UFUNCTION( BlueprintPure, Category="Editable Mesh" )
+	int32 GetEdgeConnectedPolygonCount( const FEdgeID EdgeID ) const;
+
+	/** Returns the indexed polygon connected to this edge */
+	UFUNCTION( BlueprintPure, Category="Editable Mesh" )
+	FPolygonID GetEdgeConnectedPolygon( const FEdgeID EdgeID, const int32 ConnectedPolygonNumber ) const;
+
+	/** Returns the number of polygon groups in this mesh */
+	UFUNCTION( BlueprintPure, Category="Editable Mesh" )
+	int32 GetPolygonGroupCount() const;
+
+	/** Returns whether the given polygon group ID is valid */
+	UFUNCTION( BlueprintPure, Category="Editable Mesh" )
+	bool IsValidPolygonGroup( const FPolygonGroupID PolygonGroupID ) const;
+
+	/** Returns the number of polygons in this polygon group */
+	UFUNCTION( BlueprintPure, Category="Editable Mesh" )
+	int32 GetPolygonCountInGroup( const FPolygonGroupID PolygonGroupID ) const;
+
+	/** Returns the given indexed polygon in this polygon group */
+	UFUNCTION( BlueprintPure, Category="Editable Mesh" )
+	FPolygonID GetPolygonInGroup( const FPolygonGroupID PolygonGroupID, const int32 PolygonNumber ) const;
+
+	/** Returns the material assigned to this polygon group */
+	UFUNCTION( BlueprintPure, Category="Editable Mesh" )
+	UMaterialInterface* GetPolygonGroupMaterial( const FPolygonGroupID PolygonGroupID ) const;
+
+	/** Returns the material slot name assigned to this polygon group */
+	UFUNCTION( BlueprintPure, Category="Editable Mesh" )
+	FName GetPolygonGroupMaterialSlotName( const FPolygonGroupID PolygonGroupID ) const;
+
+	/** Returns the number of polygons in this mesh */
+	UFUNCTION( BlueprintPure, Category="Editable Mesh" )
+	int32 GetPolygonCount() const;
+
+	/** Returns whether the given polygon ID is valid */
+	UFUNCTION( BlueprintPure, Category="Editable Mesh" )
+	bool IsValidPolygon( const FPolygonID PolygonID ) const;
+
+	/** Returns the polygon group this polygon is assigned to */
+	UFUNCTION( BlueprintPure, Category="Editable Mesh" )
+	FPolygonGroupID GetGroupForPolygon( const FPolygonID PolygonID ) const;
+
+	/** Returns the number of vertices on this polygon's perimeter */
+	UFUNCTION( BlueprintPure, Category="Editable Mesh" )
+	int32 GetPolygonPerimeterVertexCount( const FPolygonID PolygonID ) const;
+
+	/** Returns the indexed vertex on this polygon's perimeter */
+	UFUNCTION( BlueprintPure, Category="Editable Mesh" )
+	FVertexID GetPolygonPerimeterVertex( const FPolygonID PolygonID, const int32 PolygonVertexNumber ) const;
+
+	/** Returns the indexed vertex instance on this polygon's perimeter */
+	UFUNCTION( BlueprintPure, Category="Editable Mesh" )
+	FVertexInstanceID GetPolygonPerimeterVertexInstance( const FPolygonID PolygonID, const int32 PolygonVertexNumber ) const;
+
+	/** Returns the vertex instance attribute corresponding to the indexed vertex on the polygon's perimeter */
+	UFUNCTION( BlueprintPure, Category="Editable Mesh" )
+	FVector4 GetPolygonPerimeterVertexAttribute( const FPolygonID PolygonID, const int32 PolygonVertexNumber, const FName AttributeName, const int32 AttributeIndex ) const;
+
+	/** Returns the number of hole contours this polygon has */
+	UFUNCTION( BlueprintPure, Category="Editable Mesh" )
+	int32 GetPolygonHoleCount( const FPolygonID PolygonID ) const;
+
+	/** Returns the number of vertices on the indexed hole contour for this polygon */
+	UFUNCTION( BlueprintPure, Category="Editable Mesh" )
+	int32 GetPolygonHoleVertexCount( const FPolygonID PolygonID, const int32 HoleNumber ) const;
+
+	/** Returns the indexed vertex on the given hole contour of the polygon */
+	UFUNCTION( BlueprintPure, Category="Editable Mesh" )
+	FVertexID GetPolygonHoleVertex( const FPolygonID PolygonID, const int32 HoleNumber, const int32 PolygonVertexNumber ) const;
+
+	/** Returns the indexed vertex instance on the given hole contour of the polygon */
+	UFUNCTION( BlueprintPure, Category="Editable Mesh" )
+	FVertexInstanceID GetPolygonHoleVertexInstance( const FPolygonID PolygonID, const int32 HoleNumber, const int32 PolygonVertexNumber ) const;
+
+	/** Returns the vertex instance attribute corresponding to the indexed vertex on the given hole contour of the polygon */
+	UFUNCTION( BlueprintPure, Category="Editable Mesh" )
+	FVector4 GetPolygonHoleVertexAttribute( const FPolygonID PolygonID, const int32 HoleNumber, const int32 PolygonVertexNumber, const FName AttributeName, const int32 AttributeIndex ) const;
+
+	/** Returns the number of triangles which make up this polygon */
+	UFUNCTION( BlueprintPure, Category="Editable Mesh" )
+	int32 GetPolygonTriangulatedTriangleCount( const FPolygonID PolygonID ) const;
+
+	/** Returns the indexed triangle of the triangulated polygon */
+	UFUNCTION( BlueprintPure, Category="Editable Mesh" )
+	FMeshTriangle GetPolygonTriangulatedTriangle( const FPolygonID PolygonID, int32 PolygonTriangleNumber ) const;
 
 protected:
 
-	/** Given a set of index remappings, fixes up references to element IDs */
-	void FixUpElementIDs( const FElementIDRemappings& Remappings );
+	/** Given a set of index remappings, fixes up references in the octree */
+	void RemapOctreeIDs( const FElementIDRemappings& Remappings );
 
 	FEdgeID GetPolygonContourEdge( const FMeshPolygonContour& Contour, const int32 ContourEdgeNumber, bool& bOutEdgeWindingIsReversedForPolygon ) const;
 	void GetPolygonContourEdges( const FMeshPolygonContour& Contour, TArray<FEdgeID>& OutPolygonContourEdgeIDs ) const;
@@ -564,7 +360,6 @@ protected:
 	void SplitVertexInstanceInPolygons( const FVertexInstanceID VertexInstanceID, const TArray<FPolygonID>& PolygonIDs );
 	void ReplaceVertexInstanceInPolygons( const FVertexInstanceID OldVertexInstanceID, const FVertexInstanceID NewVertexInstanceID, const TArray<FPolygonID>& PolygonIDs );
 	float GetPolygonCornerAngleForVertex( const FPolygonID PolygonID, const FVertexID VertexID ) const;
-	FPolygonGroupID GetPolygonGroupIDFromMaterial( UMaterialInterface* Material, bool bCreateNewSectionIfNotFound );
 
 public:
 
@@ -718,7 +513,6 @@ public:
 	UFUNCTION( BlueprintPure, Category="Editable Mesh" ) void GetVertexConnectedEdges( const FVertexID VertexID, TArray<FEdgeID>& OutConnectedEdgeIDs ) const;
 	UFUNCTION( BlueprintPure, Category="Editable Mesh" ) void GetVertexConnectedPolygons( const FVertexID VertexID, TArray<FPolygonID>& OutConnectedPolygonIDs ) const;
 	UFUNCTION( BlueprintPure, Category="Editable Mesh" ) void GetVertexInstanceConnectedPolygons( const FVertexInstanceID VertexInstanceID, TArray<FPolygonID>& OutConnectedPolygonIDs ) const;
-	UFUNCTION( BlueprintPure, Category="Editable Mesh" ) int32 GetVertexInstanceConnectedPolygonCount( const FVertexInstanceID VertexInstanceID ) const;
 	UFUNCTION( BlueprintPure, Category="Editable Mesh" ) void GetVertexAdjacentVertices( const FVertexID VertexID, TArray< FVertexID >& OutAdjacentVertexIDs ) const;
 	UFUNCTION( BlueprintPure, Category="Editable Mesh" ) FEdgeID GetVertexPairEdge( const FVertexID VertexID, const FVertexID NextVertexID, bool& bOutEdgeWindingIsReversed ) const;
 	UFUNCTION( BlueprintPure, Category="Editable Mesh" ) void GetEdgeVertices( const FEdgeID EdgeID, FVertexID& OutEdgeVertexID0, FVertexID& OutEdgeVertexID1 ) const;
@@ -795,7 +589,7 @@ public:
 	UFUNCTION( BlueprintCallable, Category="Editable Mesh" ) void TriangulatePolygons( const TArray<FPolygonID>& PolygonIDs, TArray<FPolygonID>& OutNewTrianglePolygons );
 	UFUNCTION( BlueprintCallable, Category="Editable Mesh" ) void CreatePolygonGroups( const TArray<FPolygonGroupToCreate>& PolygonGroupsToCreate, TArray<FPolygonGroupID>& OutNewPolygonGroupIDs );
 	UFUNCTION( BlueprintCallable, Category="Editable Mesh" ) void DeletePolygonGroups( const TArray<FPolygonGroupID>& PolygonGroupIDs );
-	UFUNCTION( BlueprintCallable, Category="Editable Mesh" ) void AssignMaterialToPolygons( const TArray<FPolygonID>& PolygonIDs, UMaterialInterface* Material, TArray<FPolygonID>& OutNewPolygonIDs );
+	UFUNCTION( BlueprintCallable, Category="Editable Mesh" ) void AssignPolygonsToPolygonGroups( const TArray<FPolygonGroupForPolygon>& PolygonGroupForPolygons, const bool bDeleteOrphanedPolygonGroups );
 	UFUNCTION( BlueprintCallable, Category="Editable Mesh" ) void WeldVertices( const TArray<FVertexID>& VertexIDs, FVertexID& OutNewVertexID );
 	UFUNCTION( BlueprintCallable, Category="Editable Mesh" ) void TessellatePolygons( const TArray<FPolygonID>& PolygonIDs, const ETriangleTessellationMode TriangleTessellationMode, TArray<FPolygonID>& OutNewPolygonIDs );
 	UFUNCTION( BlueprintCallable, Category="Editable Mesh" ) void SetTextureCoordinateCount( const int32 NumTexCoords );
@@ -845,21 +639,24 @@ public:
 
 public:
 
+	UPROPERTY()
+	UMeshDescription* MeshDescription;
+
 	/** Each editable vertex in this mesh. */
-	TSparseArray<FMeshVertex> Vertices;
+//	TSparseArray<FMeshVertex> Vertices;
 
 	/** Sparse array of rendering vertices, that matches the vertices in the mesh vertex buffers */
-	TSparseArray<FMeshVertexInstance> VertexInstances;
+//	TSparseArray<FMeshVertexInstance> VertexInstances;
 
 	/** All editable mesh edges.  Note that some of these edges will be internal polygon edges, synthesized while
 	triangulating polygons into triangles.  Static meshes currently only support triangles. */
-	TSparseArray<FMeshEdge> Edges;
+//	TSparseArray<FMeshEdge> Edges;
 
 	/** All of the polygons in this mesh */
-	TSparseArray<FMeshPolygon> Polygons;
+//	TSparseArray<FMeshPolygon> Polygons;
 
 	/** All of the polygon groups in this mesh */
-	TSparseArray<FMeshPolygonGroup> PolygonGroups;
+//	TSparseArray<FMeshPolygonGroup> PolygonGroups;
 
 // @todo mesheditor: sort out member access. Currently StaticMesh adapter relies on accessing this stuff directly
 //protected:
diff --git a/Engine/Plugins/Runtime/EditableMesh/Source/EditableMesh/Public/EditableMeshAdapter.h b/Engine/Plugins/Runtime/EditableMesh/Source/EditableMesh/Public/EditableMeshAdapter.h
--- a/Engine/Plugins/Runtime/EditableMesh/Source/EditableMesh/Public/EditableMeshAdapter.h
+++ b/Engine/Plugins/Runtime/EditableMesh/Source/EditableMesh/Public/EditableMeshAdapter.h
@@ -47,5 +47,6 @@ public:
 	virtual void OnChangePolygonVertexInstances( const UEditableMesh* EditableMesh, const TArray<FPolygonID>& PolygonIDs ) PURE_VIRTUAL(,);
 	virtual void OnCreatePolygonGroups( const UEditableMesh* EditableMesh, const TArray<FPolygonGroupID>& PolygonGroupIDs ) PURE_VIRTUAL(,);
 	virtual void OnDeletePolygonGroups( const UEditableMesh* EditableMesh, const TArray<FPolygonGroupID>& PolygonGroupIDs ) PURE_VIRTUAL(,);
+	virtual void OnAssignPolygonsToPolygonGroups( const UEditableMesh* EditableMesh, const TArray<FPolygonGroupForPolygon>& PolygonGroupForPolygons ) PURE_VIRTUAL(,);
 	virtual void OnRetriangulatePolygons( const UEditableMesh* EditableMesh, const TArray<FPolygonID>& PolygonIDs ) PURE_VIRTUAL(,);
 };
