Build fixes and stripped dead code.

diff --git a/Engine/Plugins/Runtime/EditableMesh/Source/EditableMesh/EditableMesh.cpp b/Engine/Plugins/Runtime/EditableMesh/Source/EditableMesh/EditableMesh.cpp
--- a/Engine/Plugins/Runtime/EditableMesh/Source/EditableMesh/EditableMesh.cpp
+++ b/Engine/Plugins/Runtime/EditableMesh/Source/EditableMesh/EditableMesh.cpp
@@ -85,7 +85,7 @@ static void InvertRemapTable( TSparseArray<ElementIDType>& InvertedRemapTable, c
 
 	InvertedRemapTable.Empty( RemapTable.Num() );
 
-	for( TSparseArray<ElementIDType>::TConstIterator It( RemapTable ); It; ++It )
+	for( typename TSparseArray<ElementIDType>::TConstIterator It( RemapTable ); It; ++It )
 	{
 		InvertedRemapTable.Insert( It->GetValue(), ElementIDType( It.GetIndex() ) );
 	}
@@ -6069,7 +6069,7 @@ void UEditableMesh::ExtrudePolygons( const TArray<FPolygonID>& PolygonIDs, const
 					bool bEdgeWindingIsReversedForPolygon;
 					const FEdgeID EdgeID = this->GetPolygonPerimeterEdge( ExtrudedFrontPolygonID, PerimeterEdgeNumber, /* Out */ bEdgeWindingIsReversedForPolygon );
 
-					const FVector4 NewEdgeHardnessAttribute = bIsSharedEdge ? GetEdgeAttribute( OriginalEdgeID, UEditableMeshAttribute::EdgeIsHard(), 0 ) : FVector( 1.0f );
+					const FVector4 NewEdgeHardnessAttribute = bIsSharedEdge ? GetEdgeAttribute( OriginalEdgeID, UEditableMeshAttribute::EdgeIsHard(), 0 ) : FVector4( 1.0f, 0.0f, 0.0f, 0.0f );
 
 					FAttributesForEdge& AttributesForEdge = *new( AttributesForEdges ) FAttributesForEdge();
 					AttributesForEdge.EdgeID = EdgeID;
diff --git a/Engine/Plugins/Runtime/EditableMesh/Source/EditableMesh/Public/EditableMesh.h b/Engine/Plugins/Runtime/EditableMesh/Source/EditableMesh/Public/EditableMesh.h
--- a/Engine/Plugins/Runtime/EditableMesh/Source/EditableMesh/Public/EditableMesh.h
+++ b/Engine/Plugins/Runtime/EditableMesh/Source/EditableMesh/Public/EditableMesh.h
@@ -32,113 +32,6 @@ namespace OpenSubdiv
 }
 
 
-/**
-* Perform custom serialization for TSparseArray.
-* The default TSparseArray serialization also compacts all the elements, removing the gaps and changing the indices.
-* The indices are significant in editable meshes, hence this is a custom serializer which preserves them.
-*/
-template <typename T>
-void SerializeSparseArray( FArchive& Ar, TSparseArray<T>& Array )
-{
-	if( Ar.CustomVer( FEditableMeshCustomVersion::GUID ) < FEditableMeshCustomVersion::CustomSparseArraySerialization )
-	{
-		Ar << Array;
-	}
-	else
-	{
-		Array.CountBytes( Ar );
-
-		if( Ar.IsLoading() )
-		{
-			// Load array
-			TBitArray<> AllocatedIndices;
-			Ar << AllocatedIndices;
-
-			Array.Empty( AllocatedIndices.Num() );
-			for( auto It = TConstSetBitIterator<>( AllocatedIndices ); It; ++It )
-			{
-				Array.Insert( It.GetIndex(), T() );
-				Ar << Array[ It.GetIndex() ];
-			}
-		}
-		else
-		{
-			// Save array
-			const int32 MaxIndex = Array.GetMaxIndex();
-
-			// We have to build the TBitArray representing allocated indices by hand, as we don't have access to it from outside TSparseArray.
-			// @todo core: consider replacing TSparseArray serialization with this format.
-			TBitArray<> AllocatedIndices( false, MaxIndex );
-			for( int32 Index = 0; Index < MaxIndex; ++Index )
-			{
-				if( Array.IsAllocated( Index ) )
-				{
-					AllocatedIndices[ Index ] = true;
-				}
-			}
-			Ar << AllocatedIndices;
-
-			for( auto It = Array.CreateIterator(); It; ++It )
-			{
-				Ar << *It;
-			}
-		}
-	}
-}
-
-
-template <typename T, typename ElementIDType>
-void CompactSparseArrayElements( TSparseArray<T>& Array, TSparseArray<ElementIDType>& IndexRemap )
-{
-	static_assert( TIsDerivedFrom<ElementIDType, FElementID>::IsDerived, "Remap array type must be derived from FElementID" );
-
-	static TSparseArray<T> NewArray;
-	NewArray.Empty( Array.Num() );
-
-	IndexRemap.Empty( Array.GetMaxIndex() );
-
-	// Add valid elements into a new contiguous sparse array.  Note non-const iterator so we can move elements.
-	for( TSparseArray<T>::TIterator It( Array ); It; ++It )
-	{
-		const int32 OldElementIndex = It.GetIndex();
-
-		// @todo mesheditor: implement TSparseArray::Add( ElementType&& ) to save this obscure approach
-		const int32 NewElementIndex = NewArray.Add( T() );
-		NewArray[ NewElementIndex ] = MoveTemp( *It );
-
-		// Provide an O(1) lookup from old index to new index, used when patching up vertex references afterwards
-		IndexRemap.Insert( OldElementIndex, ElementIDType( NewElementIndex ) );
-	}
-
-	Array = MoveTemp( NewArray );
-}
-
-
-template <typename T, typename ElementIDType>
-void RemapSparseArrayElements( TSparseArray<T>& Array, const TSparseArray<ElementIDType>& IndexRemap )
-{
-	static_assert( TIsDerivedFrom<ElementIDType, FElementID>::IsDerived, "Remap array type must be derived from FElementID" );
-
-	static TSparseArray<T> NewArray;
-	NewArray.Empty( IndexRemap.GetMaxIndex() );
-
-	// Add valid elements into a new contiguous sparse array.  Note non-const iterator so we can move elements.
-	for( TSparseArray<T>::TIterator It( Array ); It; ++It )
-	{
-		const int32 OldElementIndex = It.GetIndex();
-
-		check( IndexRemap.IsAllocated( OldElementIndex ) );
-		const int32 NewElementIndex = IndexRemap[ OldElementIndex ].GetValue();
-
-		// @todo mesheditor: implement TSparseArray::Insert( ElementType&& ) to save this obscure approach
-		NewArray.Insert( NewElementIndex, T() );
-		NewArray[ NewElementIndex ] = MoveTemp( *It );
-	}
-
-	Array = MoveTemp( NewArray );
-}
-
-
 UENUM( BlueprintType )
 enum class EInsetPolygonsMode : uint8
 {
