Total revamp of mesh element attribute model.

diff --git a/Engine/Plugins/Runtime/EditableMesh/Source/EditableMesh/EditableMesh.cpp b/Engine/Plugins/Runtime/EditableMesh/Source/EditableMesh/EditableMesh.cpp
--- a/Engine/Plugins/Runtime/EditableMesh/Source/EditableMesh/EditableMesh.cpp
+++ b/Engine/Plugins/Runtime/EditableMesh/Source/EditableMesh/EditableMesh.cpp
@@ -10,6 +10,7 @@
 #include "ProfilingDebugging/ScopedTimers.h"	// For FAutoScopedDurationTimer
 #include "GeomTools.h"
 #include "mikktspace.h"	// For tangent computations
+#include "MeshAttributes.h"
 
 // =========================================================
 // OpenSubdiv support
@@ -52,6 +53,277 @@ namespace EditableMesh
 // =========================================================
 //
 
+/** Traits class which specifies whether Lerp is defined for a given type */
+template <typename T> struct TIsAttributeTypeLerpable { static const bool Value = false; };
+template <> struct TIsAttributeTypeLerpable<FVector4> { static const bool Value = true; };
+template <> struct TIsAttributeTypeLerpable<FVector> { static const bool Value = true; };
+template <> struct TIsAttributeTypeLerpable<FVector2D> { static const bool Value = true; };
+template <> struct TIsAttributeTypeLerpable<float> { static const bool Value = true; };
+template <> struct TIsAttributeTypeLerpable<int> { static const bool Value = true; };
+
+
+template <typename T, typename ElementIDType>
+static typename TEnableIf<TIsAttributeTypeLerpable<T>::Value>::Type InterpAttributeOfType( FMeshElementAttributeList& AttributeList, const FName AttributeName, const TAttributeIndicesArray<T, ElementIDType>& AttributeIndicesArray, const ElementIDType ID0, const ElementIDType ID1, const float Alpha )
+{
+	for( int32 Index = 0; Index < AttributeIndicesArray.GetNumIndices(); ++Index )
+	{
+		const TMeshAttributeArray<T, ElementIDType>& Array = AttributeIndicesArray.GetArrayForIndex( Index );
+		const T& Value0 = Array[ ID0 ];
+		const T& Value1 = Array[ ID1 ];
+
+		AttributeList.Attributes.Emplace( AttributeName, Index, 
+			FMeshElementAttributeValue( EnumHasAllFlags( AttributeIndicesArray.GetFlags(), EMeshAttributeFlags::Lerpable ) ?
+				FMath::Lerp( Value0, Value1, Alpha ) :
+				( Alpha <= 0.5f ) ? Value0 : Value1 )
+			);
+	}
+}
+
+template <typename T, typename ElementIDType>
+static typename TEnableIf<!TIsAttributeTypeLerpable<T>::Value>::Type InterpAttributeOfType( FMeshElementAttributeList& AttributeList, const FName AttributeName, const TAttributeIndicesArray<T, ElementIDType>& AttributeIndicesArray, const ElementIDType ID0, const ElementIDType ID1, const float Alpha )
+{
+	// Non-Lerpable attribute type: just choose whichever value alpha is closest to
+	for( int32 Index = 0; Index < AttributeIndicesArray.GetNumIndices(); ++Index )
+	{
+		const TMeshAttributeArray<T, ElementIDType>& Array = AttributeIndicesArray.GetArrayForIndex( Index );
+		const T& Value0 = Array[ ID0 ];
+		const T& Value1 = Array[ ID1 ];
+
+		AttributeList.Attributes.Emplace( AttributeName, Index, FMeshElementAttributeValue( ( Alpha <= 0.5f ) ? Value0 : Value1 ) );
+	}
+}
+
+
+/**
+ * Automatically set all attributes as the Lerp between the two attribute indices specified (if the attribute is marked as Lerpable),
+ * otherwise to the closest value.
+ */
+template <typename ElementIDType>
+static void InterpAllAttributes( FMeshElementAttributeList& AttributeList, const TAttributesSet<ElementIDType>& AttributesSet, const ElementIDType ID0, const ElementIDType ID1, const float Alpha )
+{
+	AttributesSet.ForEachAttributeIndicesArray(
+		[ &AttributeList, ID0, ID1, Alpha ]( const FName AttributeName, const auto& AttributeIndicesArray )
+		{
+			InterpAttributeOfType( AttributeList, AttributeName, AttributeIndicesArray, ID0, ID1, Alpha );
+		}
+	);
+}
+
+
+/**
+ * Automatically set all attributes which match the predicate as the Lerp between the two attribute indices specified (if the attribute is marked as Lerpable),
+ * otherwise to the closest value.
+ */
+template <typename ElementIDType, typename Predicate>
+static void InterpAttributesByPredicate( FMeshElementAttributeList& AttributeList, const TAttributesSet<ElementIDType>& AttributesSet, const ElementIDType ID0, const ElementIDType ID1, const float Alpha, const Predicate& Pred )
+{
+	AttributesSet.ForEachAttributeIndicesArray(
+		[ &AttributeList, ID0, ID1, Alpha, &Pred ]( const FName AttributeName, const auto& AttributeIndicesArray )
+		{
+			if( Pred( AttributeName, AttributeIndicesArray ) )
+			{
+				InterpAttributeOfType( AttributeList, AttributeName, AttributeIndicesArray, ID0, ID1, Alpha );
+			}
+		}
+	);
+}
+
+
+template <typename T, typename ElementIDType>
+static typename TEnableIf<TIsAttributeTypeLerpable<T>::Value>::Type InterpAttributeOfType( FMeshElementAttributeList& AttributeList, const FName AttributeName, const TAttributeIndicesArray<T, ElementIDType>& AttributeIndicesArray, const ElementIDType ID0, const ElementIDType ID1, const ElementIDType ID2, const FVector Weights )
+{
+	if( EnumHasAllFlags( AttributeIndicesArray.GetFlags(), EMeshAttributeFlags::Lerpable ) )
+	{
+		// Get result by interpolating between the three values according to the barycentric weight
+		for( int32 Index = 0; Index < AttributeIndicesArray.GetNumIndices(); ++Index )
+		{
+			const TMeshAttributeArray<T, ElementIDType>& Array = AttributeIndicesArray.GetArrayForIndex( Index );
+			const T& Value0 = Array[ ID0 ];
+			const T& Value1 = Array[ ID1 ];
+			const T& Value2 = Array[ ID2 ];
+
+			AttributeList.Attributes.Emplace( AttributeName, Index, FMeshElementAttributeValue( Value0 * Weights.X + Value1 * Weights.Y + Value2 * Weights.Z ) );
+		}
+	}
+	else
+	{
+		// Non-Lerpable attribute type: just choose whichever value the barycentric weight is closest to
+		const ElementIDType ClosestID = ( Weights.X > Weights.Y ) ? ( ( Weights.X > Weights.Z ) ? ID0 : ID2 ) : ( ( Weights.Y > Weights.Z ) ? ID1 : ID2 );
+
+		for( int32 Index = 0; Index < AttributeIndicesArray.GetNumIndices(); ++Index )
+		{
+			const TMeshAttributeArray<T, ElementIDType>& Array = AttributeIndicesArray.GetArrayForIndex( Index );
+			AttributeList.Attributes.Emplace( AttributeName, Index, FMeshElementAttributeValue( Array[ ClosestID ] ) );
+		}
+	}
+}
+
+template <typename T, typename ElementIDType>
+static typename TEnableIf<!TIsAttributeTypeLerpable<T>::Value>::Type InterpAttributeOfType( FMeshElementAttributeList& AttributeList, const FName AttributeName, const TAttributeIndicesArray<T, ElementIDType>& AttributeIndicesArray, const ElementIDType ID0, const ElementIDType ID1, const ElementIDType ID2, const FVector Weights )
+{
+	// Non-Lerpable attribute type: just choose whichever value the barycentric weight is closest to
+	const ElementIDType ClosestID = ( Weights.X > Weights.Y ) ? ( ( Weights.X > Weights.Z ) ? ID0 : ID2 ) : ( ( Weights.Y > Weights.Z ) ? ID1 : ID2 );
+
+	for( int32 Index = 0; Index < AttributeIndicesArray.GetNumIndices(); ++Index )
+	{
+		const TMeshAttributeArray<T, ElementIDType>& Array = AttributeIndicesArray.GetArrayForIndex( Index );
+		AttributeList.Attributes.Emplace( AttributeName, Index, FMeshElementAttributeValue( Array[ ClosestID ] ) );
+	}
+}
+
+
+/**
+ * Automatically set all attributes according to the barycentric weights of the three indices specified (if the attribute is marked as Lerpable),
+ * otherwise to the closest value.
+ */
+template <typename ElementIDType>
+static void InterpAllAttributes( FMeshElementAttributeList& AttributeList, const TAttributesSet<ElementIDType>& AttributesSet, const ElementIDType ID0, const ElementIDType ID1, const ElementIDType ID2, const FVector Weights )
+{
+	AttributesSet.ForEachAttributeIndicesArray(
+		[ &AttributeList, ID0, ID1, ID2, Weights ]( const FName AttributeName, const auto& AttributeIndicesArray )
+		{
+			InterpAttributeOfType( AttributeList, AttributeName, AttributeIndicesArray, ID0, ID1, ID2, Weights );
+		}
+	);
+}
+
+
+/**
+ * Automatically set all attributes which match the predicate according to the barycentric weights of the three indices specified (if the attribute is marked as Lerpable),
+ * otherwise to the closest value.
+ */
+template <typename ElementIDType, typename Predicate>
+static void InterpAttributesByPredicate( FMeshElementAttributeList& AttributeList, const TAttributesSet<ElementIDType>& AttributesSet, const ElementIDType ID0, const ElementIDType ID1, const ElementIDType ID2, const FVector Weights, const Predicate& Pred )
+{
+	AttributesSet.ForEachAttributeIndicesArray(
+		[ &AttributeList, ID0, ID1, ID2, Weights, &Pred ]( const FName AttributeName, const auto& AttributeIndicesArray )
+		{
+			if( Pred( AttributeName, AttributeIndicesArray ) )
+			{
+				InterpAttributeOfType( AttributeList, AttributeName, AttributeIndicesArray, ID0, ID1, ID2, Weights );
+			}
+		}
+	);
+}
+
+
+/** Compares two elements of all attribute arrays which match the given predicate, and returns whether they are all equal or not */
+template <typename ElementIDType, typename Predicate>
+static bool CompareAttributesByPredicate( const TAttributesSet<ElementIDType>& AttributesSet, const ElementIDType ID0, const ElementIDType ID1, const Predicate& Pred )
+{
+	bool bResult = true;
+	AttributesSet.ForEachAttributeIndicesArray(
+		[ ID0, ID1, &Pred, &bResult ]( const FName AttributeName, const auto& AttributeIndicesArray )
+		{
+			if( Pred( AttributeName, AttributeIndicesArray) )
+			{
+				for( int32 Index = 0; Index < AttributeIndicesArray.GetNumIndices(); ++Index )
+				{
+					auto& Array = AttributeIndicesArray.GetArrayForIndex( Index );
+					if( Array[ ID0 ] != Array[ ID1 ] )
+					{
+						bResult = false;
+						return;
+					}
+				}
+			}
+		}
+	);
+
+	return bResult;
+}
+
+
+/** Copies all element attributes with the given ID into an FMeshElementAttributeList */
+template <typename ElementIDType>
+static void BackupAllAttributes( FMeshElementAttributeList& AttributeList, const TAttributesSet<ElementIDType>& AttributesSet, const ElementIDType ElementID )
+{
+	AttributesSet.ForEachAttributeIndicesArray(
+		[ &AttributeList, ElementID ]( const FName AttributeName, const auto& AttributeIndicesArray )
+		{
+			for( int32 Index = 0; Index < AttributeIndicesArray.GetNumIndices(); ++Index )
+			{
+				const auto& Value = AttributeIndicesArray.GetArrayForIndex( Index )[ ElementID ];
+
+				AttributeList.Attributes.Emplace( AttributeName, Index, FMeshElementAttributeValue( Value ) );
+//				UE_LOG( LogTemp, Log, TEXT( "Name: %s, Index: %d, Value: %s" ), *AttributeName.ToString(), Index, *FMeshElementAttributeValue( Value ).ToString() );
+			}
+		}
+	);
+}
+
+
+/** Copies all element attributes which match the given predicate with the given ID into an FMeshElementAttributeList */
+template <typename ElementIDType, typename Predicate>
+static void BackupAttributesByPredicate( FMeshElementAttributeList& AttributeList, const TAttributesSet<ElementIDType>& AttributesSet, const ElementIDType ElementID, const Predicate& Pred )
+{
+	AttributesSet.ForEachAttributeIndicesArray(
+		[ &AttributeList, ElementID, &Pred ]( const FName AttributeName, const auto& AttributeIndicesArray )
+		{
+			if( Pred( AttributeName, AttributeIndicesArray ) )
+			{
+				for( int32 Index = 0; Index < AttributeIndicesArray.GetNumIndices(); ++Index )
+				{
+					const auto& Value = AttributeIndicesArray.GetArrayForIndex( Index )[ ElementID ];
+
+					AttributeList.Attributes.Emplace( AttributeName, Index, FMeshElementAttributeValue( Value ) );
+//					UE_LOG( LogTemp, Log, TEXT( "Name: %s, Index: %d, Value: %s" ), *AttributeName.ToString(), Index, *FMeshElementAttributeValue( Value ).ToString() );
+				}
+			}
+		}
+	);
+}
+
+
+/** Applies an FMeshElementAttributeData to an element attribute with the given ID */
+template <typename ElementIDType>
+static void ApplyAttribute( TAttributesSet<ElementIDType>& AttributesSet, const FMeshElementAttributeData& AttributeData, const ElementIDType ElementID )
+{
+	AttributeData.AttributeValue.Visit(
+		[ &AttributesSet, &AttributeData, ElementID ]( const auto& Value )
+		{
+			using AttributeType = TDecay<decltype( Value )>::Type;
+			AttributesSet.SetAttribute<AttributeType>( ElementID, AttributeData.AttributeName, AttributeData.AttributeIndex, Value );
+		}
+	);
+}
+
+
+/** Applies all attributes in an FMeshElementAttributeList to an element attribute with the given ID */
+template <typename ElementIDType>
+static void ApplyAttributes( TAttributesSet<ElementIDType>& AttributesSet, const FMeshElementAttributeList& AttributeList, const ElementIDType ElementID )
+{
+	for( const FMeshElementAttributeData& AttributeData : AttributeList.Attributes )
+	{
+		ApplyAttribute( AttributesSet, AttributeData, ElementID );
+	}
+}
+
+
+/** Given a FMeshElementAttributeList of attributes to be set, first backup the current values of those attributes into a new list, and then apply the new values */
+template <typename ElementIDType>
+static void BackupAttributesInList( FMeshElementAttributeList& BackupList, const FMeshElementAttributeList& SourceList, const TAttributesSet<ElementIDType>& AttributesSet, const ElementIDType ElementID )
+{
+	BackupList.Attributes.Reserve( SourceList.Attributes.Num() );
+
+	for( const FMeshElementAttributeData& AttributeData : SourceList.Attributes )
+	{
+		AttributeData.AttributeValue.Visit(
+			[ &BackupList, &AttributesSet, &AttributeData, ElementID ]( const auto& Value )
+			{
+				const FName& AttributeName = AttributeData.AttributeName;
+				const int32 AttributeIndex = AttributeData.AttributeIndex;
+
+				// Put current value into BackupList
+				using AttributeType = TDecay<decltype( Value )>::Type;
+				const AttributeType OldValue = AttributesSet.GetAttribute<AttributeType>( ElementID, AttributeName, AttributeIndex );
+				BackupList.Attributes.Emplace( AttributeName, AttributeIndex, FMeshElementAttributeValue( OldValue ) );
+			}
+		);
+	}
+}
+
+
 UEditableMesh::UEditableMesh()
 	: bAllowUndo( false ),
 	  bAllowCompact( false ),
@@ -616,21 +888,6 @@ FPolygonID UEditableMesh::GetPolygonInGroup( const FPolygonGroupID PolygonGroupI
 }
 
 
-UMaterialInterface* UEditableMesh::GetPolygonGroupMaterial( const FPolygonGroupID PolygonGroupID ) const
-{
-	const FMeshPolygonGroup& PolygonGroup = GetMeshDescription()->GetPolygonGroup( PolygonGroupID );
-	UMaterialInterface* Material = Cast<UMaterialInterface>( PolygonGroup.MaterialAsset.ResolveObject() );
-	check( Material );
-	return Material;
-}
-
-
-FName UEditableMesh::GetPolygonGroupMaterialSlotName( const FPolygonGroupID PolygonGroupID ) const
-{
-	return GetMeshDescription()->GetPolygonGroup( PolygonGroupID ).MaterialSlotName;
-}
-
-
 int32 UEditableMesh::GetPolygonCount() const
 {
 	return GetMeshDescription()->Polygons().Num();
@@ -667,12 +924,6 @@ FVertexInstanceID UEditableMesh::GetPolygonPerimeterVertexInstance( const FPolyg
 }
 
 
-FVector4 UEditableMesh::GetPolygonPerimeterVertexAttribute( const FPolygonID PolygonID, const int32 PolygonVertexNumber, const FName AttributeName, const int32 AttributeIndex ) const
-{
-	return GetVertexInstanceAttribute( GetPolygonPerimeterVertexInstance( PolygonID, PolygonVertexNumber ), AttributeName, AttributeIndex );
-}
-
-
 int32 UEditableMesh::GetPolygonHoleCount( const FPolygonID PolygonID ) const
 {
 	return GetMeshDescription()->GetPolygon( PolygonID ).HoleContours.Num();
@@ -697,12 +948,6 @@ FVertexInstanceID UEditableMesh::GetPolygonHoleVertexInstance( const FPolygonID
 }
 
 
-FVector4 UEditableMesh::GetPolygonHoleVertexAttribute( const FPolygonID PolygonID, const int32 HoleNumber, const int32 PolygonVertexNumber, const FName AttributeName, const int32 AttributeIndex ) const
-{
-	return GetVertexInstanceAttribute( GetPolygonHoleVertexInstance( PolygonID, HoleNumber, PolygonVertexNumber ), AttributeName, AttributeIndex );
-}
-
-
 int32 UEditableMesh::GetPolygonTriangulatedTriangleCount( const FPolygonID PolygonID ) const
 {
 	return GetMeshDescription()->GetPolygon( PolygonID ).Triangles.Num();
@@ -715,105 +960,13 @@ FMeshTriangle UEditableMesh::GetPolygonTriangulatedTriangle( const FPolygonID Po
 }
 
 
-FVector4 UEditableMesh::GetVertexAttribute( const FVertexID VertexID, const FName AttributeName, const int32 AttributeIndex ) const
-{
-	const FMeshVertex& Vertex = GetMeshDescription()->GetVertex( VertexID );
-
-	if( AttributeName == UEditableMeshAttribute::VertexPosition() )
-	{
-		check( AttributeIndex == 0 );	// Only one position is supported
-		return FVector4( Vertex.VertexPosition, 0.0f );
-	}
-	else if( AttributeName == UEditableMeshAttribute::VertexCornerSharpness() )
-	{
-		check( AttributeIndex == 0 );	// Only one softness value is supported
-		return FVector4( Vertex.CornerSharpness, 0.0f, 0.0f, 0.0f );
-	}
-
-	checkf( 0, TEXT( "UEditableMesh::GetVertexAttribute() called with unrecognized vertex attribute name: %s (index: %i)" ), *AttributeName.ToString(), AttributeIndex );
-	return FVector4( 0.0f );
-}
-
-
-FVector4 UEditableMesh::GetVertexInstanceAttribute( const FVertexInstanceID VertexInstanceID, const FName AttributeName, const int32 AttributeIndex ) const
-{
-	const FMeshVertexInstance& VertexInstance = GetMeshDescription()->GetVertexInstance( VertexInstanceID );
-
-	if( AttributeName == UEditableMeshAttribute::VertexPosition() )
-	{
-		check( AttributeIndex == 0 );	// Only one position is supported
-		return GetVertexAttribute( VertexInstance.VertexID, AttributeName, AttributeIndex );
-	}
-	else if( AttributeName == UEditableMeshAttribute::VertexNormal() )
-	{
-		check( AttributeIndex == 0 );	// Only one normal is supported
-		return FVector4( VertexInstance.Normal, 0.0f );
-	}
-	else if( AttributeName == UEditableMeshAttribute::VertexTangent() )
-	{
-		check( AttributeIndex == 0 );	// Only one tangent is supported
-		return FVector4( VertexInstance.Tangent, 0.0f );
-	}
-	else if( AttributeName == UEditableMeshAttribute::VertexBinormalSign() )
-	{
-		check( AttributeIndex == 0 );	// Only one basis determinant sign is supported
-		return FVector4( VertexInstance.BinormalSign );
-	}
-	else if( AttributeName == UEditableMeshAttribute::VertexTextureCoordinate() )
-	{
-		const int32 TextureCoordinateIndex = AttributeIndex;
-		if( TextureCoordinateIndex < VertexInstance.VertexUVs.Num() )
-		{
-			const FVector2D TextureCoordinate = VertexInstance.VertexUVs[ TextureCoordinateIndex ];
-			return FVector4( TextureCoordinate, FVector2D::ZeroVector );
-		}
-		else
-		{
-			return FVector4( 0.0f );
-		}
-	}
-	else if( AttributeName == UEditableMeshAttribute::VertexColor() )
-	{
-		check( AttributeIndex == 0 );
-		return FVector4( FLinearColor( VertexInstance.Color ) );
-	}
-
-	checkf( 0, TEXT( "UEditableMesh::GetVertexInstanceAttribute() called with unrecognized vertex instance attribute name: %s (index: %i)" ), *AttributeName.ToString(), AttributeIndex );
-	return FVector4( 0.0f );
-}
-
-
-FVector4 UEditableMesh::GetEdgeAttribute( const FEdgeID EdgeID, const FName AttributeName, const int32 AttributeIndex ) const
-{
-	const FMeshEdge& Edge = GetMeshDescription()->GetEdge( EdgeID );
-
-	if( AttributeName == UEditableMeshAttribute::EdgeIsHard() )
-	{
-		check( AttributeIndex == 0 );	// Only one edge is hard flag is supported
-
-		return FVector4( Edge.bIsHardEdge ? 1.0f : 0.0f, 0.0f, 0.0f, 0.0f );
-	}
-	else if( AttributeName == UEditableMeshAttribute::EdgeCreaseSharpness() )
-	{
-		check( AttributeIndex == 0 );	// Only one edge crease sharpness is supported
-
-		return FVector4( Edge.CreaseSharpness, 0.0f, 0.0f, 0.0f );
-	}
-
-	checkf( 0, TEXT( "UEditableMesh::GetEdgeAttribute() called with unrecognized edge attribute name: %s (index: %i)" ), *AttributeName.ToString(), AttributeIndex );
-	return FVector4( 0.0f );
-}
-
-
-void UEditableMesh::SetVertexAttribute( const FVertexID VertexID, const FName AttributeName, const int32 AttributeIndex, const FVector4 AttributeValue )
+void UEditableMesh::SetVertexAttribute( const FVertexID VertexID, const FMeshElementAttributeData& Attribute )
 {
-	FMeshVertex& Vertex = GetMeshDescription()->GetVertex( VertexID );
+	ApplyAttribute( GetMeshDescription()->VertexAttributes(), Attribute, VertexID );
 
-	if( AttributeName == UEditableMeshAttribute::VertexPosition() )
+	// @todo mesh editor: create attribute flag which signifies that the octree needs rebuilding
+	if( Attribute.AttributeName == MeshAttribute::Vertex::Position )
 	{
-		check( AttributeIndex == 0 );	// Only one position is supported
-		Vertex.VertexPosition = AttributeValue;
-
 		// Update spatial database
 		if( Octree.IsValid() )
 		{
@@ -832,95 +985,63 @@ void UEditableMesh::SetVertexAttribute( const FVertexID VertexID, const FName At
 			}
 		}
 	}
-	else if( AttributeName == UEditableMeshAttribute::VertexCornerSharpness() )
-	{
-		check( AttributeIndex == 0 );	// Only one softness value is supported
-		Vertex.CornerSharpness = AttributeValue.X;
-	}
-	else
-	{
-		checkf( 0, TEXT( "UEditableMesh::SetVertexAttribute() called with unrecognized vertex attribute name: %s (index: %i)" ), *AttributeName.ToString(), AttributeIndex );
-	}
 
 	for( UEditableMeshAdapter* Adapter : Adapters )
 	{
-		Adapter->OnSetVertexAttribute( this, VertexID, AttributeName, AttributeIndex, AttributeValue );
+		Adapter->OnSetVertexAttribute( this, VertexID, Attribute );
 	}
 }
 
 
-void UEditableMesh::SetVertexInstanceAttribute( const FVertexInstanceID VertexInstanceID, const FName AttributeName, const int32 AttributeIndex, const FVector4 AttributeValue )
+void UEditableMesh::SetVertexInstanceAttribute( const FVertexInstanceID VertexInstanceID, const FMeshElementAttributeData& Attribute )
 {
-	FMeshVertexInstance& VertexInstance = GetMeshDescription()->GetVertexInstance( VertexInstanceID );
-
-	if( AttributeName == UEditableMeshAttribute::VertexPosition() )
-	{
-		check( AttributeIndex == 0 );	// Only one position is supported
-		SetVertexAttribute( VertexInstance.VertexID, AttributeName, AttributeIndex, AttributeValue );
-	}
-	else if( AttributeName == UEditableMeshAttribute::VertexNormal() )
-	{
-		check( AttributeIndex == 0 );	// Only one normal is supported
-		VertexInstance.Normal = AttributeValue;
-	}
-	else if( AttributeName == UEditableMeshAttribute::VertexTangent() )
-	{
-		check( AttributeIndex == 0 );	// Only one tangent is supported
-		VertexInstance.Tangent = AttributeValue;
-	}
-	else if( AttributeName == UEditableMeshAttribute::VertexBinormalSign() )
-	{
-		check( AttributeIndex == 0 );	// Only one basis determinant sign is supported
-		VertexInstance.BinormalSign = AttributeValue.X;
-	}
-	else if( AttributeName == UEditableMeshAttribute::VertexTextureCoordinate() )
-	{
-		if( AttributeIndex < VertexInstance.VertexUVs.Num() )
-		{
-			VertexInstance.VertexUVs[ AttributeIndex ] = FVector2D( AttributeValue.X, AttributeValue.Y );
-		}
-	}
-	else if( AttributeName == UEditableMeshAttribute::VertexColor() )
-	{
-		check( AttributeIndex == 0 );
-		VertexInstance.Color = FLinearColor( AttributeValue.X, AttributeValue.Y, AttributeValue.Z, AttributeValue.W );
-	}
-	else
-	{
-		checkf( 0, TEXT( "UEditableMesh::SetVertexInstanceAttribute() called with unrecognized vertex instance attribute name: %s (index: %i)" ), *AttributeName.ToString(), AttributeIndex );
-	}
+	ApplyAttribute( GetMeshDescription()->VertexInstanceAttributes(), Attribute, VertexInstanceID );
 
 	for( UEditableMeshAdapter* Adapter : Adapters )
 	{
-		Adapter->OnSetVertexInstanceAttribute( this, VertexInstanceID, AttributeName, AttributeIndex, AttributeValue );
+		Adapter->OnSetVertexInstanceAttribute( this, VertexInstanceID, Attribute );
 	}
 }
 
 
-void UEditableMesh::SetEdgeAttribute( const FEdgeID EdgeID, const FName AttributeName, const int32 AttributeIndex, const FVector4 AttributeValue, bool bIsUndo )
+void UEditableMesh::SetEdgeAttribute( const FEdgeID EdgeID, const FMeshElementAttributeData& Attribute, bool bIsUndo )
 {
-	FMeshEdge& Edge = GetMeshDescription()->GetEdge( EdgeID );
+	ApplyAttribute( GetMeshDescription()->EdgeAttributes(), Attribute, EdgeID );
 
-	if( AttributeName == UEditableMeshAttribute::EdgeIsHard() )
+	if( Attribute.AttributeName == MeshAttribute::Edge::IsHard )
 	{
-		check( AttributeIndex == 0 );	// Only one edge is hard flag is supported
-		SetEdgeHardness( EdgeID, AttributeValue.X != 0.0f, bIsUndo );
+//		SetEdgeHardness( EdgeID, Attribute.AttributeValue.GetValue<bool>(), bIsUndo );
 
 		// @todo mesheditor: this needs to potentially merge/split vertex instances and rebuild normals
 	}
-	else if( AttributeName == UEditableMeshAttribute::EdgeCreaseSharpness() )
+
+	for( UEditableMeshAdapter* Adapter : Adapters )
 	{
-		check( AttributeIndex == 0 );	// Only one edge crease sharpness is supported
-		Edge.CreaseSharpness = AttributeValue.X;
+		Adapter->OnSetEdgeAttribute( this, EdgeID, Attribute );
 	}
-	else
+}
+
+
+void UEditableMesh::SetPolygonAttribute( const FPolygonID PolygonID, const FMeshElementAttributeData& Attribute )
+{
+	ApplyAttribute( GetMeshDescription()->PolygonAttributes(), Attribute, PolygonID );
+
+	for( UEditableMeshAdapter* Adapter : Adapters )
 	{
-		checkf( 0, TEXT( "UEditableMesh::SetEdgeAttribute() called with unrecognized edge attribute name: %s (index: %i)" ), *AttributeName.ToString(), AttributeIndex );
+		// @todo: implement me!
+//		Adapter->OnSetPolygonAttribute( this, PolygonID, Attribute );
 	}
+}
+
+
+void UEditableMesh::SetPolygonGroupAttribute( const FPolygonGroupID PolygonGroupID, const FMeshElementAttributeData& Attribute )
+{
+	ApplyAttribute( GetMeshDescription()->PolygonGroupAttributes(), Attribute, PolygonGroupID );
 
 	for( UEditableMeshAdapter* Adapter : Adapters )
 	{
-		Adapter->OnSetEdgeAttribute( this, EdgeID, AttributeName, AttributeIndex, AttributeValue );
+		// @todo: implement me!
+//		Adapter->OnSetPolygonGroupAttribute( this, PolygonGroupID, Attribute );
 	}
 }
 
@@ -983,27 +1104,14 @@ void UEditableMesh::SplitVertexInstanceInPolygons( const FVertexInstanceID Verte
 {
 	UE_LOG( LogEditableMesh, Verbose, TEXT( "SplitVertexInstanceInPolygons: %s %s" ), *VertexInstanceID.ToString(), *LogHelpers::ArrayToString( PolygonIDs ) );
 
-	const FMeshVertexInstance& VertexInstance = GetMeshDescription()->GetVertexInstance( VertexInstanceID );
-
 	// Create a new vertex instance copied from the one passed in
 
 	static TArray<FVertexInstanceToCreate> VertexInstancesToCreate;
 	VertexInstancesToCreate.Reset();
 	VertexInstancesToCreate.SetNum( 1 );
-	VertexInstancesToCreate[ 0 ].VertexID = VertexInstance.VertexID;
+	VertexInstancesToCreate[ 0 ].VertexID = GetMeshDescription()->GetVertexInstanceVertex( VertexInstanceID );
 
-	for( const FName AttributeName : UEditableMesh::GetValidVertexInstanceAttributes() )
-	{
-		const int32 MaxAttributeIndex = GetMaxAttributeIndex( AttributeName );
-		for( int32 AttributeIndex = 0; AttributeIndex < MaxAttributeIndex; ++AttributeIndex )
-		{
-			VertexInstancesToCreate[ 0 ].VertexInstanceAttributes.Attributes.Emplace(
-				AttributeName,
-				AttributeIndex,
-				GetVertexInstanceAttribute( VertexInstanceID, AttributeName, AttributeIndex )
-			);
-		}
-	}
+	BackupAllAttributes( VertexInstancesToCreate[ 0 ].VertexInstanceAttributes, GetMeshDescription()->VertexInstanceAttributes(), VertexInstanceID );
 
 	static TArray<FVertexInstanceID> NewVertexInstanceIDs;
 	CreateVertexInstances( VertexInstancesToCreate, NewVertexInstanceIDs );
@@ -1126,211 +1234,147 @@ void UEditableMesh::GetVertexInstanceConnectedPolygonsInSameGroup( const FVertex
 
 void UEditableMesh::SetEdgeHardness( const FEdgeID EdgeID, const bool bIsHard, const bool bIsUndo )
 {
+	// @todo mesheditor: revisit splitting/welding of identical vertex indices
+#if 0
 	const FVertexArray& Vertices = GetMeshDescription()->Vertices();
 	const FVertexInstanceArray& VertexInstances = GetMeshDescription()->VertexInstances();
 
 	FMeshEdge& Edge = GetMeshDescription()->GetEdge( EdgeID );
 
-	if( Edge.bIsHardEdge != bIsHard )
+	if( !bIsUndo )
 	{
-		Edge.bIsHardEdge = bIsHard;
+		UE_LOG( LogEditableMesh, Verbose, TEXT( "SetEdgeHardness (changed): %s %s" ), *EdgeID.ToString(), *LogHelpers::BoolToString( bIsHard ) );
 
-		if( !bIsUndo )
+		if( bIsHard )
 		{
-			UE_LOG( LogEditableMesh, Verbose, TEXT( "SetEdgeHardness (changed): %s %s" ), *EdgeID.ToString(), *LogHelpers::BoolToString( bIsHard ) );
+			// Setting a hard edge potentially requires vertex instances to be split
 
-			if( bIsHard )
+			for( const FVertexID VertexID : Edge.VertexIDs )
 			{
-				// Setting a hard edge potentially requires vertex instances to be split
+				const FMeshVertex& Vertex = Vertices[ VertexID ];
 
-				for( const FVertexID VertexID : Edge.VertexIDs )
-				{
-					const FMeshVertex& Vertex = Vertices[ VertexID ];
+				// For each end of the edge, go through all vertex instances, determining if any need to be split.
+				// Take a copy of the vertex instances because splitting them will mutate the list we are iterating.
+				static TArray<FVertexInstanceID> VertexInstanceIDsCopy;
+				VertexInstanceIDsCopy = Vertex.VertexInstanceIDs;
 
-					// For each end of the edge, go through all vertex instances, determining if any need to be split.
-					// Take a copy of the vertex instances because splitting them will mutate the list we are iterating.
-					static TArray<FVertexInstanceID> VertexInstanceIDsCopy;
-					VertexInstanceIDsCopy = Vertex.VertexInstanceIDs;
+				// Iterate all the vertex instances for each endpoint.
+				// We are looking for vertex instances which need splitting into two different instances
+				// (because the hard edge will require separate normals for each of its polygon smoothing groups).
+				for( const FVertexInstanceID VertexInstanceID : VertexInstanceIDsCopy )
+				{
+					const FMeshVertexInstance& VertexInstance = VertexInstances[ VertexInstanceID ];
 
-					// Iterate all the vertex instances for each endpoint.
-					// We are looking for vertex instances which need splitting into two different instances
-					// (because the hard edge will require separate normals for each of its polygon smoothing groups).
-					for( const FVertexInstanceID VertexInstanceID : VertexInstanceIDsCopy )
+					// Take a copy of the array of polygons connected to this instance.
+					static TArray<FPolygonID> PolygonsSharingThisVertex;
+					PolygonsSharingThisVertex = VertexInstance.ConnectedPolygons;
+					bool bFirstTime = true;
+					while( PolygonsSharingThisVertex.Num() > 0 )
 					{
-						const FMeshVertexInstance& VertexInstance = VertexInstances[ VertexInstanceID ];
-
-						// Take a copy of the array of polygons connected to this instance.
-						static TArray<FPolygonID> PolygonsSharingThisVertex;
-						PolygonsSharingThisVertex = VertexInstance.ConnectedPolygons;
-						bool bFirstTime = true;
-						while( PolygonsSharingThisVertex.Num() > 0 )
+						// For the next polygon in the array, determine all other polygons in the same smoothing group which share this vertex instance.
+						static TArray<FPolygonID> PolygonsInSameSmoothingGroup;
+						GetVertexInstanceConnectedPolygonsInSameGroup( VertexInstanceID, PolygonsSharingThisVertex[ 0 ], PolygonsInSameSmoothingGroup );
+
+						// Check that all polygons in the smoothing group are attached to this vertex instance, and remove them from the master list of polygons
+						// connected to this instance. If a polygon in the smoothing group is not attached to this vertex instance, it's because it's the other
+						// side of a UV seam and hence has a distinct vertex instance.
+						for( const FPolygonID PolygonInSmoothingGroup : PolygonsInSameSmoothingGroup )
 						{
-							// For the next polygon in the array, determine all other polygons in the same smoothing group which share this vertex instance.
-							static TArray<FPolygonID> PolygonsInSameSmoothingGroup;
-							GetVertexInstanceConnectedPolygonsInSameGroup( VertexInstanceID, PolygonsSharingThisVertex[ 0 ], PolygonsInSameSmoothingGroup );
-
-							// Check that all polygons in the smoothing group are attached to this vertex instance, and remove them from the master list of polygons
-							// connected to this instance. If a polygon in the smoothing group is not attached to this vertex instance, it's because it's the other
-							// side of a UV seam and hence has a distinct vertex instance.
-							for( const FPolygonID PolygonInSmoothingGroup : PolygonsInSameSmoothingGroup )
-							{
-								verify( PolygonsSharingThisVertex.Remove( PolygonInSmoothingGroup ) == 1 );
-							}
-
-							// First group which we extract: do nothing - they can keep their existing instance ID.
-							// Subsequent times round the loop, we create a new vertex instance copied from the original one, and replace connected polygon vertices with it.
-							if( !bFirstTime )
-							{
-								SplitVertexInstanceInPolygons( VertexInstanceID, PolygonsInSameSmoothingGroup );
-							}
+							verify( PolygonsSharingThisVertex.Remove( PolygonInSmoothingGroup ) == 1 );
+						}
 
-							bFirstTime = false;
+						// First group which we extract: do nothing - they can keep their existing instance ID.
+						// Subsequent times round the loop, we create a new vertex instance copied from the original one, and replace connected polygon vertices with it.
+						if( !bFirstTime )
+						{
+							SplitVertexInstanceInPolygons( VertexInstanceID, PolygonsInSameSmoothingGroup );
 						}
+
+						bFirstTime = false;
 					}
 				}
 			}
-			else
-			{
-				// Setting a soft edge potentially requires vertex instances to be merged
+		}
+		else
+		{
+			// Setting a soft edge potentially requires vertex instances to be merged
 
-				static TArray<FVertexInstanceID> VertexInstancesToDelete;
-				VertexInstancesToDelete.Reset();
+			static TArray<FVertexInstanceID> VertexInstancesToDelete;
+			VertexInstancesToDelete.Reset();
 
-				for( const FVertexID VertexID : Edge.VertexIDs )
+			for( const FVertexID VertexID : Edge.VertexIDs )
+			{
+				// For each end of the edge which has been softened...
+				const FMeshVertex& Vertex = Vertices[ VertexID ];
+
+				// ...iterate through pairs of vertex instances, looking for potential to merge them
+				for( int32 IndexA = 0; IndexA < Vertex.VertexInstanceIDs.Num() - 1; ++IndexA )
 				{
-					// For each end of the edge which has been softened...
-					const FMeshVertex& Vertex = Vertices[ VertexID ];
+					const FVertexInstanceID VertexInstanceIDA = Vertex.VertexInstanceIDs[ IndexA ];
+					const FMeshVertexInstance& VertexInstanceA = VertexInstances[ VertexInstanceIDA ];
 
-					// ...iterate through pairs of vertex instances, looking for potential to merge them
-					for( int32 IndexA = 0; IndexA < Vertex.VertexInstanceIDs.Num() - 1; ++IndexA )
+					// If vertex instance isn't connected to any polygon, we can't deduce anything about its smoothing group, so skip to the next one
+					if( VertexInstanceA.ConnectedPolygons.Num() == 0 )
 					{
-						const FVertexInstanceID VertexInstanceIDA = Vertex.VertexInstanceIDs[ IndexA ];
-						const FMeshVertexInstance& VertexInstanceA = VertexInstances[ VertexInstanceIDA ];
+						continue;
+					}
 
-						// If vertex instance isn't connected to any polygon, we can't deduce anything about its smoothing group, so skip to the next one
-						if( VertexInstanceA.ConnectedPolygons.Num() == 0 )
-						{
+					// Get mergeable set of attributes for first vertex instance
+					static TArray<FMeshElementAttributeData> AttributesA;
+					GetMergedNamedVertexInstanceAttributeData( GetMergeableVertexInstanceAttributes(), Vertex.VertexInstanceIDs[ IndexA ], TArray<FMeshElementAttributeData>(), AttributesA );
+
+					for( int32 IndexB = IndexA + 1; IndexB < Vertex.VertexInstanceIDs.Num(); ++IndexB )
+					{
+						const FVertexInstanceID VertexInstanceIDB = Vertex.VertexInstanceIDs[ IndexB ];
+
+						// If this vertex instance has been marked for deletion in a previous iteration, skip it
+						if( VertexInstancesToDelete.Contains( VertexInstanceIDB ) )
+						{
 							continue;
 						}
 
-						// Get mergeable set of attributes for first vertex instance
-						static TArray<FMeshElementAttributeData> AttributesA;
-						GetMergedNamedVertexInstanceAttributeData( GetMergeableVertexInstanceAttributes(), Vertex.VertexInstanceIDs[ IndexA ], TArray<FMeshElementAttributeData>(), AttributesA );
+						// Get mergeable set of attributes for second vertex instance
+						static TArray<FMeshElementAttributeData> AttributesB;
+						GetMergedNamedVertexInstanceAttributeData( GetMergeableVertexInstanceAttributes(), Vertex.VertexInstanceIDs[ IndexB ], TArray<FMeshElementAttributeData>(), AttributesB );
 
-						for( int32 IndexB = IndexA + 1; IndexB < Vertex.VertexInstanceIDs.Num(); ++IndexB )
+						// If both vertex instances have equal mergeable attributes, they can potentially be merged
+						if( AreAttributeListsEqual( AttributesA, AttributesB ) )
 						{
-							const FVertexInstanceID VertexInstanceIDB = Vertex.VertexInstanceIDs[ IndexB ];
+							const FMeshVertexInstance& VertexInstanceB = VertexInstances[ VertexInstanceIDB ];
 
-							// If this vertex instance has been marked for deletion in a previous iteration, skip it
-							if( VertexInstancesToDelete.Contains( VertexInstanceIDB ) )
-							{
-								continue;
-							}
+							// Check to see if they are in the same smoothing group, using an arbitary connected polygon in the first vertex instance as the pivot point.
+							static TArray<FVertexInstanceID> PossibleVertexInstances;
+							GetVertexInstancesInSameSoftEdgedGroup( VertexInstanceA.VertexID, VertexInstanceA.ConnectedPolygons[ 0 ], false, PossibleVertexInstances );
 
-							// Get mergeable set of attributes for second vertex instance
-							static TArray<FMeshElementAttributeData> AttributesB;
-							GetMergedNamedVertexInstanceAttributeData( GetMergeableVertexInstanceAttributes(), Vertex.VertexInstanceIDs[ IndexB ], TArray<FMeshElementAttributeData>(), AttributesB );
-
-							// If both vertex instances have equal mergeable attributes, they can potentially be merged
-							if( AreAttributeListsEqual( AttributesA, AttributesB ) )
+							// If the smoothing group contains both vertex instances, they can be merged.
+							if( PossibleVertexInstances.Contains( VertexInstanceIDA ) && PossibleVertexInstances.Contains( VertexInstanceIDB ) )
 							{
-								const FMeshVertexInstance& VertexInstanceB = VertexInstances[ VertexInstanceIDB ];
-
-								// Check to see if they are in the same smoothing group, using an arbitary connected polygon in the first vertex instance as the pivot point.
-								static TArray<FVertexInstanceID> PossibleVertexInstances;
-								GetVertexInstancesInSameSoftEdgedGroup( VertexInstanceA.VertexID, VertexInstanceA.ConnectedPolygons[ 0 ], false, PossibleVertexInstances );
-
-								// If the smoothing group contains both vertex instances, they can be merged.
-								if( PossibleVertexInstances.Contains( VertexInstanceIDA ) && PossibleVertexInstances.Contains( VertexInstanceIDB ) )
-								{
-									// Change occurrences of VertexInstanceB for VertexInstanceA in VertexInstanceB's connected polygons.
-									// Note, this will cause VertexInstanceA's connected polygons list to be added to (at the end).
-									// This works because we are evaluating the number of connected polygons each time round the loop.
-									ReplaceVertexInstanceInPolygons( VertexInstanceIDB, VertexInstanceIDA, VertexInstanceB.ConnectedPolygons );
-
-									// This will also cause VertexInstanceB to be disconnected from all polygons.
-									// We mark the vertex instance for deletion here, but do not delete it until the end, as to do so would interrupt iterating through vertex instances.
-									check( VertexInstanceB.ConnectedPolygons.Num() == 0 );
-									VertexInstancesToDelete.Add( VertexInstanceIDB );
-								}
+								// Change occurrences of VertexInstanceB for VertexInstanceA in VertexInstanceB's connected polygons.
+								// Note, this will cause VertexInstanceA's connected polygons list to be added to (at the end).
+								// This works because we are evaluating the number of connected polygons each time round the loop.
+								ReplaceVertexInstanceInPolygons( VertexInstanceIDB, VertexInstanceIDA, VertexInstanceB.ConnectedPolygons );
+
+								// This will also cause VertexInstanceB to be disconnected from all polygons.
+								// We mark the vertex instance for deletion here, but do not delete it until the end, as to do so would interrupt iterating through vertex instances.
+								check( VertexInstanceB.ConnectedPolygons.Num() == 0 );
+								VertexInstancesToDelete.Add( VertexInstanceIDB );
 							}
 						}
 					}
 				}
-
-				// Delete orphaned vertex instances
-				const bool bDeleteOrphanedVertices = false;
-				DeleteVertexInstances( VertexInstancesToDelete, bDeleteOrphanedVertices );
 			}
 
-			UE_LOG( LogEditableMesh, Verbose, TEXT( "* SetEdgeHardness returned" ) );
+			// Delete orphaned vertex instances
+			const bool bDeleteOrphanedVertices = false;
+			DeleteVertexInstances( VertexInstancesToDelete, bDeleteOrphanedVertices );
 		}
 
-		// Mark polygons adjacent to this edge as candidates for normal/tangent recalculation
-		PolygonsPendingNewTangentBasis.Append( Edge.ConnectedPolygons );
-	}
-}
-
-
-const TArray<FName>& UEditableMesh::GetValidVertexAttributes()
-{
-	static TArray<FName> ValidVertexAttributes;
-	if( ValidVertexAttributes.Num() == 0 )
-	{
-		ValidVertexAttributes.Add( UEditableMeshAttribute::VertexPosition() );
-		ValidVertexAttributes.Add( UEditableMeshAttribute::VertexCornerSharpness() );
-	}
-	return ValidVertexAttributes;
-}
-
-
-const TArray<FName>& UEditableMesh::GetValidVertexInstanceAttributes()
-{
-	static TArray<FName> ValidVertexInstanceAttributes;
-	if( ValidVertexInstanceAttributes.Num() == 0 )
-	{
-		ValidVertexInstanceAttributes.Add( UEditableMeshAttribute::VertexNormal() );
-		ValidVertexInstanceAttributes.Add( UEditableMeshAttribute::VertexTangent() );
-		ValidVertexInstanceAttributes.Add( UEditableMeshAttribute::VertexBinormalSign() );
-		ValidVertexInstanceAttributes.Add( UEditableMeshAttribute::VertexTextureCoordinate() );
-		ValidVertexInstanceAttributes.Add( UEditableMeshAttribute::VertexColor() );
-	}
-	return ValidVertexInstanceAttributes;
-}
-
-
-const TArray<FName>& UEditableMesh::GetValidEdgeAttributes()
-{
-	static TArray<FName> ValidEdgeAttributes;
-	if( ValidEdgeAttributes.Num() == 0 )
-	{
-		ValidEdgeAttributes.Add( UEditableMeshAttribute::EdgeIsHard() );
-		ValidEdgeAttributes.Add( UEditableMeshAttribute::EdgeCreaseSharpness() );
+		UE_LOG( LogEditableMesh, Verbose, TEXT( "* SetEdgeHardness returned" ) );
 	}
-	return ValidEdgeAttributes;
-}
-
 
-const TArray<FName>& UEditableMesh::GetValidPolygonAttributes()
-{
-	static TArray<FName> ValidPolygonAttributes;
-	if( ValidPolygonAttributes.Num() == 0 )
-	{
-		ValidPolygonAttributes.Add( UEditableMeshAttribute::PolygonNormal() );
-		ValidPolygonAttributes.Add( UEditableMeshAttribute::PolygonCenter() );
-	}
-	return ValidPolygonAttributes;
-}
-
-
-int32 UEditableMesh::GetMaxAttributeIndex( const FName AttributeName ) const
-{
-	if( AttributeName == UEditableMeshAttribute::VertexTextureCoordinate() )
-	{
-		return GetTextureCoordinateCount();
-	}
-
-	return 1;
+	// Mark polygons adjacent to this edge as candidates for normal/tangent recalculation
+	PolygonsPendingNewTangentBasis.Append( Edge.ConnectedPolygons );
+#endif
 }
 
 
@@ -1714,9 +1758,11 @@ FBox UEditableMesh::ComputeBoundingBox() const
 	FBox BoundingBox;
 	BoundingBox.Init();
 
+	const TVertexAttributeArray<FVector>& VertexPositions = GetMeshDescription()->VertexAttributes().GetAttributes<FVector>( MeshAttribute::Vertex::Position );
+
 	for( const FVertexID VertexID : GetMeshDescription()->Vertices().GetElementIDs() )
 	{
-		BoundingBox += GetVertexAttribute( VertexID, UEditableMeshAttribute::VertexPosition(), 0 );
+		BoundingBox += VertexPositions[ VertexID ];
 	}
 
 	return BoundingBox;
@@ -1733,10 +1779,11 @@ FBoxSphereBounds UEditableMesh::ComputeBoundingBoxAndSphere() const
 	// Calculate the bounding sphere, using the center of the bounding box as the origin.
 	BoundingBoxAndSphere.SphereRadius = 0.0f;
 
+	const TVertexAttributeArray<FVector>& VertexPositions = GetMeshDescription()->VertexAttributes().GetAttributes<FVector>( MeshAttribute::Vertex::Position );
+
 	for( const FVertexID VertexID : GetMeshDescription()->Vertices().GetElementIDs() )
 	{
-		const FVector VertexPosition = GetVertexAttribute( VertexID, UEditableMeshAttribute::VertexPosition(), 0 );
-		BoundingBoxAndSphere.SphereRadius = FMath::Max( ( VertexPosition - BoundingBoxAndSphere.Origin ).Size(), BoundingBoxAndSphere.SphereRadius );
+		BoundingBoxAndSphere.SphereRadius = FMath::Max( ( VertexPositions[ VertexID ] - BoundingBoxAndSphere.Origin ).Size(), BoundingBoxAndSphere.SphereRadius );
 	}
 
 	return BoundingBoxAndSphere;
@@ -1750,10 +1797,11 @@ FVector UEditableMesh::ComputePolygonCenter( const FPolygonID PolygonID ) const
 	static TArray<FVertexID> PerimeterVertexIDs;
 	GetPolygonPerimeterVertices( PolygonID, /* Out */ PerimeterVertexIDs );
 
+	const TVertexAttributeArray<FVector>& VertexPositions = GetMeshDescription()->VertexAttributes().GetAttributes<FVector>( MeshAttribute::Vertex::Position );
+
 	for( const FVertexID VertexID : PerimeterVertexIDs )
 	{
-		const FVector Position = GetVertexAttribute( VertexID, UEditableMeshAttribute::VertexPosition(), 0 );
-		Centroid += Position;
+		Centroid += VertexPositions[ VertexID ];
 	}
 
 	return Centroid / (float)PerimeterVertexIDs.Num();
@@ -1776,14 +1824,16 @@ FPlane UEditableMesh::ComputePolygonPlane( const FPolygonID PolygonID ) const
 	static TArray<FVertexID> PerimeterVertexIDs;
 	GetPolygonPerimeterVertices( PolygonID, /* Out */ PerimeterVertexIDs );
 
+	const TVertexAttributeArray<FVector>& VertexPositions = GetMeshDescription()->VertexAttributes().GetAttributes<FVector>( MeshAttribute::Vertex::Position );
+
 	// Use 'Newell's Method' to compute a robust 'best fit' plane from the vertices of this polygon
 	for( int32 VertexNumberI = PerimeterVertexIDs.Num() - 1, VertexNumberJ = 0; VertexNumberJ < PerimeterVertexIDs.Num(); VertexNumberI = VertexNumberJ, VertexNumberJ++ )
 	{
 		const FVertexID VertexIDI = PerimeterVertexIDs[ VertexNumberI ];
-		const FVector PositionI = GetVertexAttribute( VertexIDI, UEditableMeshAttribute::VertexPosition(), 0 );
+		const FVector PositionI = VertexPositions[ VertexIDI ];
 
 		const FVertexID VertexIDJ = PerimeterVertexIDs[ VertexNumberJ ];
-		const FVector PositionJ = GetVertexAttribute( VertexIDJ, UEditableMeshAttribute::VertexPosition(), 0 );
+		const FVector PositionJ = VertexPositions[ VertexIDJ ];
 
 		Centroid += PositionJ;
 
@@ -1833,9 +1883,12 @@ void UEditableMesh::RefreshOpenSubdiv()
 				{
 					OsdCornerVertexIndices.Reset();
 					OsdCornerWeights.Reset();
+
+					const TVertexAttributeArray<float>& VertexCornerSharpnesses = GetMeshDescription()->VertexAttributes().GetAttributes<float>( MeshAttribute::Vertex::CornerSharpness );
+
 					for( const FVertexID VertexID : GetMeshDescription()->Vertices().GetElementIDs() )
 					{
-						const float VertexCornerSharpness = GetVertexAttribute( VertexID, UEditableMeshAttribute::VertexCornerSharpness(), 0 ).X;
+						const float VertexCornerSharpness = VertexCornerSharpnesses[ VertexID ];
 						if( VertexCornerSharpness > SMALL_NUMBER )
 						{
 							// This vertex is (at least partially) a subdivision corner
@@ -1849,9 +1902,12 @@ void UEditableMesh::RefreshOpenSubdiv()
 				{
 					OsdCreaseVertexIndexPairs.Reset();
 					OsdCreaseWeights.Reset();
+
+					const TEdgeAttributeArray<float>& EdgeCreaseSharpnesses = GetMeshDescription()->EdgeAttributes().GetAttributes<float>( MeshAttribute::Edge::CreaseSharpness );
+
 					for( const FEdgeID EdgeID : GetMeshDescription()->Edges().GetElementIDs() )
 					{
-						const float EdgeCreaseSharpness = GetEdgeAttribute( EdgeID, UEditableMeshAttribute::EdgeCreaseSharpness(), 0 ).X;
+						const float EdgeCreaseSharpness = EdgeCreaseSharpnesses[ EdgeID ];
 						if( EdgeCreaseSharpness > SMALL_NUMBER )
 						{
 							// This edge is (at least partially) creased
@@ -2038,12 +2094,17 @@ void UEditableMesh::GenerateOpenSubdivLimitSurfaceData()
 				//       sparse array optimization.)
 				const int32 VertexArraySize = GetMeshDescription()->Vertices().GetArraySize();
 				VertexPositions.SetNum( VertexArraySize, false );
+
+				// @todo mesheditor perf: note that now we expose vertex positions as an array of FVector, so this is just copying from one array to another.
+				// This needs to be refactored to not require this.
+				const TVertexAttributeArray<FVector>& SrcVertexPositions = MeshDescription->VertexAttributes().GetAttributes<FVector>( MeshAttribute::Vertex::Position );
+
 				for( int32 VertexNumber = 0; VertexNumber < VertexArraySize; ++VertexNumber )
 				{
 					const FVertexID VertexID( VertexNumber );
 					if( IsValidVertex( VertexID ) )
 					{
-						VertexPositions[ VertexNumber ] = GetVertexAttribute( VertexID, UEditableMeshAttribute::VertexPosition(), 0 );
+						VertexPositions[ VertexNumber ] = SrcVertexPositions[ VertexID ];
 					}
 					else
 					{
@@ -2060,24 +2121,25 @@ void UEditableMesh::GenerateOpenSubdivLimitSurfaceData()
 				FVarVertexDatas.Reset();
 				FVarVertexDatas.Reserve( OsdFVarIndicesPerFace.Num() );
 
+				const TAttributeIndicesArray<FVector2D, FVertexInstanceID>& VertexUVs = GetMeshDescription()->VertexInstanceAttributes().GetAttributesSet<FVector2D>( MeshAttribute::VertexInstance::TextureCoordinate );
+				const TVertexInstanceAttributeArray<FVector4>& VertexColors = GetMeshDescription()->VertexInstanceAttributes().GetAttributes<FVector4>( MeshAttribute::VertexInstance::Color );
+
 				int32 NumPolygonsSoFar = 0;
 				for( const FPolygonGroupID PolygonGroupID : GetMeshDescription()->PolygonGroups().GetElementIDs() )
 				{
 					FirstPolygonNumberForPolygonGroups.Add( NumPolygonsSoFar );
-					const int32 PolygonCountInGroup = GetPolygonCountInGroup( PolygonGroupID );
-					NumPolygonsSoFar += PolygonCountInGroup;
 
-					for( int32 PolygonGroupPolygonNumber = 0; PolygonGroupPolygonNumber < PolygonCountInGroup; ++PolygonGroupPolygonNumber )
-					{
-						const FPolygonID PolygonID( GetPolygonInGroup( PolygonGroupID, PolygonGroupPolygonNumber ) );
+					const TArray<FPolygonID>& PolygonIDs = GetMeshDescription()->GetPolygonGroupPolygons( PolygonGroupID );
+					NumPolygonsSoFar += PolygonIDs.Num();
 
-						const int32 PerimeterVertexCount = GetPolygonPerimeterVertexCount( PolygonID );
-						for( int32 PerimeterVertexNumber = 0; PerimeterVertexNumber < PerimeterVertexCount; ++PerimeterVertexNumber )
+					for( const FPolygonID PolygonID : PolygonIDs )
+					{
+						for( const FVertexInstanceID VertexInstanceID : GetMeshDescription()->GetPolygonPerimeterVertexInstances( PolygonID ) )
 						{
 							FOsdFVarVertexData& FVarVertexData = *new( FVarVertexDatas ) FOsdFVarVertexData();
-							FVarVertexData.TextureCoordinates[ 0 ] = TextureCoordinateCount > 0 ? FVector2D( GetPolygonPerimeterVertexAttribute( PolygonID, PerimeterVertexNumber, UEditableMeshAttribute::VertexTextureCoordinate(), 0 ) ) : FVector2D::ZeroVector;
-							FVarVertexData.TextureCoordinates[ 1 ] = TextureCoordinateCount > 1 ? FVector2D( GetPolygonPerimeterVertexAttribute( PolygonID, PerimeterVertexNumber, UEditableMeshAttribute::VertexTextureCoordinate(), 1 ) ) : FVector2D::ZeroVector;
-							FVarVertexData.VertexColor = FLinearColor( GetPolygonPerimeterVertexAttribute( PolygonID, PerimeterVertexNumber, UEditableMeshAttribute::VertexColor(), 0 ) );
+							FVarVertexData.TextureCoordinates[ 0 ] = TextureCoordinateCount > 0 ? VertexUVs.GetArrayForIndex( 0 )[ VertexInstanceID ] : FVector2D::ZeroVector;
+							FVarVertexData.TextureCoordinates[ 1 ] = TextureCoordinateCount > 1 ? VertexUVs.GetArrayForIndex( 1 )[ VertexInstanceID ] : FVector2D::ZeroVector;
+							FVarVertexData.VertexColor = FLinearColor( VertexColors[ VertexInstanceID ] );
 						}
 					}
 				}
@@ -2481,8 +2543,8 @@ void UEditableMesh::ComputePolygonTriangulation( const FPolygonID PolygonID, TAr
 	// @todo mesheditor holes: Does not support triangles with holes yet!
 	// @todo mesheditor: Perhaps should always attempt to triangulate by splitting polygons along the shortest edge, for better determinism.
 
-	const FMeshPolygon& Polygon = GetMeshDescription()->GetPolygon( PolygonID );
-	const TArray<FVertexInstanceID>& PolygonVertexInstanceIDs = Polygon.PerimeterContour.VertexInstanceIDs;
+//	const FMeshPolygon& Polygon = GetMeshDescription()->GetPolygon( PolygonID );
+	const TArray<FVertexInstanceID>& PolygonVertexInstanceIDs = GetMeshDescription()->GetPolygonPerimeterVertexInstances( PolygonID );
 
 	// Polygon must have at least three vertices/edges
 	const int32 PolygonVertexCount = PolygonVertexInstanceIDs.Num();
@@ -2496,7 +2558,9 @@ void UEditableMesh::ComputePolygonTriangulation( const FPolygonID PolygonID, TAr
 	// in the polygon.  This will just save us having to iterate later on.
 	static TArray<int32> PrevVertexNumbers, NextVertexNumbers;
 	static TArray<FVector> VertexPositions;
+
 	{
+		const TVertexAttributeArray<FVector>& MeshVertexPositions = GetMeshDescription()->VertexAttributes().GetAttributes<FVector>( MeshAttribute::Vertex::Position );
 		PrevVertexNumbers.SetNumUninitialized( PolygonVertexCount, false );
 		NextVertexNumbers.SetNumUninitialized( PolygonVertexCount, false );
 		VertexPositions.SetNumUninitialized( PolygonVertexCount, false );
@@ -2506,7 +2570,7 @@ void UEditableMesh::ComputePolygonTriangulation( const FPolygonID PolygonID, TAr
 			PrevVertexNumbers[ VertexNumber ] = VertexNumber - 1;
 			NextVertexNumbers[ VertexNumber ] = VertexNumber + 1;
 
-			VertexPositions[ VertexNumber ] = GetVertexInstanceAttribute( PolygonVertexInstanceIDs[ VertexNumber ], UEditableMeshAttribute::VertexPosition(), 0 );
+			VertexPositions[ VertexNumber ] = MeshVertexPositions[ GetMeshDescription()->GetVertexInstanceVertex( PolygonVertexInstanceIDs[ VertexNumber ] ) ];
 		}
 		PrevVertexNumbers[ 0 ] = PolygonVertexCount - 1;
 		NextVertexNumbers[ PolygonVertexCount - 1 ] = 0;
@@ -2612,12 +2676,14 @@ bool UEditableMesh::ComputeBarycentricWeightForPointOnPolygon( const FPolygonID
 	const FMeshPolygon& Polygon = GetMeshDescription()->GetPolygon( PolygonID );
 	const int32 TriangleCount = Polygon.Triangles.Num();
 
+	const TVertexAttributeArray<FVector>& VertexPositions = GetMeshDescription()->VertexAttributes().GetAttributes<FVector>( MeshAttribute::Vertex::Position );
+
 	// Figure out which triangle the incoming point is within
 	for( const FMeshTriangle& Triangle : Polygon.Triangles )
 	{
-		const FVector TriangleVertex0Position = GetVertexInstanceAttribute( Triangle.VertexInstanceID0, UEditableMeshAttribute::VertexPosition(), 0 );
-		const FVector TriangleVertex1Position = GetVertexInstanceAttribute( Triangle.VertexInstanceID1, UEditableMeshAttribute::VertexPosition(), 0 );
-		const FVector TriangleVertex2Position = GetVertexInstanceAttribute( Triangle.VertexInstanceID2, UEditableMeshAttribute::VertexPosition(), 0 );
+		const FVector TriangleVertex0Position = VertexPositions[ GetMeshDescription()->GetVertexInstanceVertex( Triangle.VertexInstanceID0 ) ];
+		const FVector TriangleVertex1Position = VertexPositions[ GetMeshDescription()->GetVertexInstanceVertex( Triangle.VertexInstanceID1 ) ];
+		const FVector TriangleVertex2Position = VertexPositions[ GetMeshDescription()->GetVertexInstanceVertex( Triangle.VertexInstanceID2 ) ];
 
 		// Calculate the barycentric weights for the triangle's verts and determine if the point lies within its bounds.
 		OutTriangleVertexWeights = FMath::ComputeBaryCentric2D( PointOnPolygon, TriangleVertex0Position, TriangleVertex1Position, TriangleVertex2Position );
@@ -2634,30 +2700,6 @@ bool UEditableMesh::ComputeBarycentricWeightForPointOnPolygon( const FPolygonID
 }
 
 
-void UEditableMesh::ComputeTextureCoordinatesForPointOnPolygon( const FPolygonID PolygonID, const FVector PointOnPolygon, bool& bOutWereTextureCoordinatesFound, TArray<FVector4>& OutInterpolatedTextureCoordinates )
-{
-	bOutWereTextureCoordinatesFound = false;
-	OutInterpolatedTextureCoordinates.Reset();
-
-	FMeshTriangle Triangle;
-	FVector TriangleVertexWeights;
-	if( ComputeBarycentricWeightForPointOnPolygon( PolygonID, PointOnPolygon, Triangle, TriangleVertexWeights ) )
-	{
-		OutInterpolatedTextureCoordinates.SetNum( TextureCoordinateCount, false );
-
-		for( int32 TextureCoordinateIndex = 0; TextureCoordinateIndex < TextureCoordinateCount; ++TextureCoordinateIndex )
-		{
-			OutInterpolatedTextureCoordinates[ TextureCoordinateIndex ] =
-				TriangleVertexWeights.X * GetVertexInstanceAttribute( Triangle.VertexInstanceID0, UEditableMeshAttribute::VertexTextureCoordinate(), TextureCoordinateIndex ) +
-				TriangleVertexWeights.Y * GetVertexInstanceAttribute( Triangle.VertexInstanceID1, UEditableMeshAttribute::VertexTextureCoordinate(), TextureCoordinateIndex ) +
-				TriangleVertexWeights.Z * GetVertexInstanceAttribute( Triangle.VertexInstanceID2, UEditableMeshAttribute::VertexTextureCoordinate(), TextureCoordinateIndex );
-		}
-
-		bOutWereTextureCoordinatesFound = true;
-	}
-}
-
-
 void UEditableMesh::SetSubdivisionCount( const int32 NewSubdivisionCount )
 {
 	// @todo mesheditor subdiv: Really, instead of a custom FChange type for this type of change, we could create a change that
@@ -2694,15 +2736,18 @@ void UEditableMesh::MoveVertices( const TArray<FVertexToMove>& VerticesToMove )
 	static TArray<FAttributesForVertex> VertexAttributesToSet;
 	VertexAttributesToSet.Reset();
 
+	const TVertexAttributeArray<FVector>& VertexPositions = GetMeshDescription()->VertexAttributes().GetAttributes<FVector>( MeshAttribute::Vertex::Position );
+
 	for( const FVertexToMove& VertexToMove : VerticesToMove )
 	{
-		const FVector CurrentPosition = GetVertexAttribute( VertexToMove.VertexID, UEditableMeshAttribute::VertexPosition(), 0 );
+		const FVector CurrentPosition = VertexPositions[ VertexToMove.VertexID ];
 
 		if( VertexToMove.NewVertexPosition != CurrentPosition )
 		{
-			FAttributesForVertex& AttributesForVertex = *new( VertexAttributesToSet ) FAttributesForVertex();
+			VertexAttributesToSet.Emplace();
+			FAttributesForVertex& AttributesForVertex = VertexAttributesToSet.Last();
 			AttributesForVertex.VertexID = VertexToMove.VertexID;
-			FMeshElementAttributeData& VertexAttribute = *new( AttributesForVertex.VertexAttributes.Attributes ) FMeshElementAttributeData( UEditableMeshAttribute::VertexPosition(), 0, VertexToMove.NewVertexPosition );
+			AttributesForVertex.VertexAttributes.Attributes.Emplace( MeshAttribute::Vertex::Position, 0, FMeshElementAttributeValue( VertexToMove.NewVertexPosition ) );
 
 			// All of the polygons that share this vertex will need new normals
 			static TArray<FPolygonID> ConnectedPolygonRefs;
@@ -2859,14 +2904,16 @@ void UEditableMesh::SplitEdge( const FEdgeID EdgeID, const TArray<float>& Splits
 	// NOTE: The incoming splits should always be between 0.0 and 1.0, representing progress along 
 	//       the edge from the edge's first vertex toward it's other vertex.  The order doesn't matter (we'll sort them.)
 
-	check( Splits.Num() > 0 );
+	const int32 NumNewVertices = Splits.Num();
+
+	check( NumNewVertices > 0 );
 
 	// Sort the split values smallest to largest.  We'll be adding a strip of vertices for each split, and the
 	// indices for those new vertices need to be in order.
 	static TArray<float> SortedSplits;
 	SortedSplits.Reset();
 	SortedSplits = Splits;
-	if( SortedSplits.Num() > 1 )
+	if( NumNewVertices > 1 )
 	{
 		SortedSplits.Sort();
 	}
@@ -2874,59 +2921,30 @@ void UEditableMesh::SplitEdge( const FEdgeID EdgeID, const TArray<float>& Splits
 	FVertexID OriginalEdgeVertexIDs[ 2 ];
 	GetEdgeVertices( EdgeID, /* Out */ OriginalEdgeVertexIDs[0], /* Out */ OriginalEdgeVertexIDs[1] );
 
-	static TArray<FVector> NewVertexPositions;
-	NewVertexPositions.Reset();
-	NewVertexPositions.AddUninitialized( SortedSplits.Num() );
-
-	for( int32 NewVertexNumber = 0; NewVertexNumber < NewVertexPositions.Num(); ++NewVertexNumber )
-	{
-		const float Split = SortedSplits[ NewVertexNumber ];
-		check( Split >= 0.0f && Split <= 1.0f );
-
-		NewVertexPositions[ NewVertexNumber ] = FMath::Lerp(
-			FVector( GetVertexAttribute( OriginalEdgeVertexIDs[ 0 ], UEditableMeshAttribute::VertexPosition(), 0 ) ),
-			FVector( GetVertexAttribute( OriginalEdgeVertexIDs[ 1 ], UEditableMeshAttribute::VertexPosition(), 0 ) ),
-			Split );
-	}
-
 	// Split the edge, and connect the vertex to the polygons that share the two new edges
 	const FVertexID OriginalEdgeFarVertexID = OriginalEdgeVertexIDs[ 1 ];
 
-	static TArray<FMeshElementAttributeData> OriginalEdgeAttributes;
-	OriginalEdgeAttributes.Reset();
-	{
-		for( const FName AttributeName : UEditableMesh::GetValidEdgeAttributes() )
-		{
-			const int32 MaxAttributeIndex = GetMaxAttributeIndex( AttributeName );
-			for( int32 AttributeIndex = 0; AttributeIndex < MaxAttributeIndex; ++AttributeIndex )
-			{
-				FMeshElementAttributeData& EdgeAttribute = *new( OriginalEdgeAttributes ) FMeshElementAttributeData(
-					AttributeName,
-					AttributeIndex,
-					GetEdgeAttribute( EdgeID, AttributeName, AttributeIndex ) );
-			}
-		}
-	}
+	// Get attributes for the edge being split: they will be copied into the new edges created
+	static FMeshElementAttributeList EdgeAttributeList;
+	EdgeAttributeList.Attributes.Reset();
+	BackupAllAttributes( EdgeAttributeList, GetMeshDescription()->EdgeAttributes(), EdgeID );
 
 
-	// Create new vertices.  Their texture coordinates will be undefined, for now.  We'll set that later by using the connected 
-	// polygon's vertex data.  Also, the vertex's normals and tangents will be automatically generated.
+	// Create new vertices, and get new attributes by interpolating values between the edge delimiters.
 	static TArray<FVertexID> NewVertexIDs;
 	NewVertexIDs.Reset();
 	{
 		OutNewVertexIDs.Reset();
-		OutNewVertexIDs.Reserve( NewVertexPositions.Num() );
+		OutNewVertexIDs.Reserve( NumNewVertices );
 
 		static TArray<FVertexToCreate> VerticesToCreate;
 		VerticesToCreate.Reset();
-		VerticesToCreate.Reserve( NewVertexPositions.Num() );
+		VerticesToCreate.Reserve( NumNewVertices );
 
-		for( int32 NewVertexNumber = 0; NewVertexNumber < NewVertexPositions.Num(); ++NewVertexNumber )
+		for( int32 NewVertexNumber = 0; NewVertexNumber < NumNewVertices; ++NewVertexNumber )
 		{
-			const FVector& NewVertexPosition = NewVertexPositions[ NewVertexNumber ];
-
 			FVertexToCreate& VertexToCreate = *new( VerticesToCreate )FVertexToCreate();
-			VertexToCreate.VertexAttributes.Attributes.Add( FMeshElementAttributeData( UEditableMeshAttribute::VertexPosition(), 0, NewVertexPosition ) );
+			InterpAllAttributes( VertexToCreate.VertexAttributes, GetMeshDescription()->VertexAttributes(), OriginalEdgeVertexIDs[ 0 ], OriginalEdgeVertexIDs[ 1 ], SortedSplits[ NewVertexNumber ] );
 		}
 
 		CreateVertices( VerticesToCreate, /* Out */ NewVertexIDs );
@@ -2934,56 +2952,100 @@ void UEditableMesh::SplitEdge( const FEdgeID EdgeID, const TArray<float>& Splits
 		OutNewVertexIDs.Append( NewVertexIDs );
 	}
 
-	// Figure out which polygons are connected to the original edge.  Those polygons are going to need their
-	// connected vertices updated later on.
+
+	// Create new vertex instances.
 	struct FAffectedPolygonEdge
 	{
-		FPolygonID PolygonID;
-		int32 PolygonVertexNumbers[ 2 ];
+		int32 PerimeterVertexIndex;
+		int32 NewVertexInstanceArrayIndex;
+		bool bReversed;
 	};
-	static TArray<FAffectedPolygonEdge> AffectedPolygonEdges;
-	AffectedPolygonEdges.Reset();
-
-	static TArray<FPolygonID> AffectedPolygons;
-	AffectedPolygons.Reset();
 
+	static TMap<TTuple<FVertexInstanceID, FVertexInstanceID>, int32> VertexInstancePairsToIndex;
+	static TMap<FPolygonID, FAffectedPolygonEdge> PolygonToAffectedEdge;
+	static TArray<FVertexInstanceToCreate> VertexInstancesToCreate;
+	static TArray<FVertexInstanceID> NewVertexInstanceIDs;
+	VertexInstancePairsToIndex.Reset();
+	PolygonToAffectedEdge.Reset();
+	VertexInstancesToCreate.Reset();
 	{
-		const int32 ConnectedPolygonCount = GetEdgeConnectedPolygonCount( EdgeID );
-		for( int32 ConnectedPolygonNumber = 0; ConnectedPolygonNumber < ConnectedPolygonCount; ++ConnectedPolygonNumber )
-		{
-			const FPolygonID PolygonID = GetEdgeConnectedPolygon( EdgeID, ConnectedPolygonNumber );
+		int32 VertexInstanceIndex = 0;
 
-			// Figure out which of this polygon's vertices are part of the edge we're adding a vertex to
-			// @todo mesheditor hole: Add support for polygon holes, too!
-			static TArray<FVertexID> PolygonPerimeterVertexIDs;
-			PolygonPerimeterVertexIDs.Reset();
-			GetPolygonPerimeterVertices( PolygonID, /* Out */ PolygonPerimeterVertexIDs );
+		// Iterate through each polygon connected to the edge being split
+		for( const FPolygonID ConnectedPolygonID : GetMeshDescription()->GetEdgeConnectedPolygons( EdgeID ) )
+		{
+			const TArray<FVertexInstanceID>& PerimeterVertexInstances = GetMeshDescription()->GetPolygonPerimeterVertexInstances( ConnectedPolygonID );
+			const int32 NumVertexInstances = PerimeterVertexInstances.Num();
 
-			int32 PolygonVertexNumberForEdgeVertices[ 2 ] = { INDEX_NONE, INDEX_NONE };
-			for( int32 PerimeterVertexNumber = 0; PerimeterVertexNumber < PolygonPerimeterVertexIDs.Num(); ++PerimeterVertexNumber )
+			// Find vertex indices around the perimeter corresponding to the vertices defining the edge
+			int32 VertexIndex = 0;
+			int32 VertexIndices[ 2 ] = { INDEX_NONE, INDEX_NONE };
+			for( const FVertexInstanceID PerimeterVertexInstance : PerimeterVertexInstances )
 			{
-				const FVertexID PerimeterVertexID = PolygonPerimeterVertexIDs[ PerimeterVertexNumber ];
-				if( PerimeterVertexID == OriginalEdgeVertexIDs[ 0 ] )
+				if( GetMeshDescription()->GetVertexInstanceVertex( PerimeterVertexInstance ) == OriginalEdgeVertexIDs[ 0 ] )
 				{
-					// Okay, this polygon vertex is part of our edge (first edge vertex)
-					PolygonVertexNumberForEdgeVertices[ 0 ] = PerimeterVertexNumber;
+					VertexIndices[ 0 ] = VertexIndex;
 				}
-				else if( PerimeterVertexID == OriginalEdgeVertexIDs[ 1 ] )
+				else if( GetMeshDescription()->GetVertexInstanceVertex( PerimeterVertexInstance ) == OriginalEdgeVertexIDs[ 1 ] )
 				{
-					// Okay, this polygon vertex is part of our edge (second edge vertex)
-					PolygonVertexNumberForEdgeVertices[ 1 ] = PerimeterVertexNumber;
+					VertexIndices[ 1 ] = VertexIndex;
 				}
+
+				VertexIndex++;
 			}
-			check( PolygonVertexNumberForEdgeVertices[ 0 ] != INDEX_NONE && PolygonVertexNumberForEdgeVertices[ 1 ] != INDEX_NONE );
 
-			FAffectedPolygonEdge& AffectedPolygonEdge = *new( AffectedPolygonEdges ) FAffectedPolygonEdge();
-			AffectedPolygonEdge.PolygonID = PolygonID;
-			AffectedPolygonEdge.PolygonVertexNumbers[ 0 ] = PolygonVertexNumberForEdgeVertices[ 0 ];
-			AffectedPolygonEdge.PolygonVertexNumbers[ 1 ] = PolygonVertexNumberForEdgeVertices[ 1 ];
+			// Check the vertex indices are valid, and are adjacent (taking into account wrap-around)
+			check( VertexIndices[ 0 ] != INDEX_NONE && VertexIndices[ 1 ] != INDEX_NONE );
+			check( FMath::Abs( VertexIndices[ 0 ] - VertexIndices[ 1 ] ) == 1 || FMath::Abs( VertexIndices[ 0 ] - VertexIndices[ 1 ] ) == NumVertexInstances - 1 );
+
+			// Get actual vertex instance IDs of the edge ends (so we can reuse new vertex instances where appropriate)
+			const FVertexInstanceID VertexInstanceID0 = PerimeterVertexInstances[ VertexIndices[ 0 ] ];
+			const FVertexInstanceID VertexInstanceID1 = PerimeterVertexInstances[ VertexIndices[ 1 ] ];
 
-			// @todo mesheditor perf: Not sure if we actually need to AddUnique() here -- can an edge be connected to the same polygon more than once? (tube shape). Seems sketchy if so?  And if so, GetEdgeConnectedPolygonCount() probably should have only reported 1 connection anyway
-			AffectedPolygons.AddUnique( AffectedPolygonEdge.PolygonID );
+			// See if this vertex instance pair has already had its intermediate vertices created
+			const int32* IndexForVertexInstancePair = VertexInstancePairsToIndex.Find( MakeTuple( VertexInstanceID0, VertexInstanceID1 ) );
+			const int32 NewVertexInstanceArrayIndex = IndexForVertexInstancePair ? *IndexForVertexInstancePair : VertexInstanceIndex;
+
+			if( !IndexForVertexInstancePair )
+			{
+				// If not, register it in the map
+				VertexInstancePairsToIndex.Add( MakeTuple( VertexInstanceID0, VertexInstanceID1 ), VertexInstanceIndex );
+				VertexInstancePairsToIndex.Add( MakeTuple( VertexInstanceID1, VertexInstanceID0 ), VertexInstanceIndex );
+
+				// And create the vertex instances, interpolating attributes
+				VertexInstancesToCreate.Reserve( VertexInstancesToCreate.Num() + NumNewVertices );
+				for( int32 NewVertexNumber = 0; NewVertexNumber < NumNewVertices; ++NewVertexNumber )
+				{
+					VertexInstancesToCreate.Emplace();
+					FVertexInstanceToCreate& VertexInstanceToCreate = VertexInstancesToCreate.Last();
+					VertexInstanceToCreate.VertexID = NewVertexIDs[ NewVertexNumber ];
+
+					InterpAllAttributes( VertexInstanceToCreate.VertexInstanceAttributes,
+										 GetMeshDescription()->VertexInstanceAttributes(),
+										 VertexInstanceID0,
+										 VertexInstanceID1,
+										 SortedSplits[ NewVertexNumber ] );										 
+				}
+
+				VertexInstanceIndex += NumNewVertices;
+			}
+
+			// Add details for this polygon edge, used to determine how new vertex instances should be inserted.
+			// A reversed edge is one where either:
+			// - The 1st perimeter index is greater than the 2nd; OR
+			// - The two indices are not consecutive (i.e. at the wrap around point);
+			// but not both.
+			const bool bReversed = ( ( VertexIndices[ 0 ] > VertexIndices[ 1 ] ) ^ ( FMath::Abs( VertexIndices[ 0 ] - VertexIndices[ 1 ] ) > 1 ) );
+
+			FAffectedPolygonEdge AffectedPolygonEdge;
+			AffectedPolygonEdge.PerimeterVertexIndex = ( bReversed ? VertexIndices[ 1 ] : VertexIndices[ 0 ] ) + 1;
+			AffectedPolygonEdge.NewVertexInstanceArrayIndex = NewVertexInstanceArrayIndex;
+			AffectedPolygonEdge.bReversed = bReversed;
+
+			PolygonToAffectedEdge.Add( ConnectedPolygonID, AffectedPolygonEdge );
 		}
+
+		CreateVertexInstances( VertexInstancesToCreate, NewVertexInstanceIDs );
 	}
 
 	{
@@ -3005,7 +3067,7 @@ void UEditableMesh::SplitEdge( const FEdgeID EdgeID, const TArray<float>& Splits
 
 	// Create new edges.  One for each of the new vertex positions passed in.
 	{
-		const int32 NewEdgeCount = NewVertexPositions.Num();
+		const int32 NewEdgeCount = NumNewVertices;
 
 		static TArray<FEdgeToCreate> EdgesToCreate;
 		EdgesToCreate.Reset();
@@ -3017,10 +3079,10 @@ void UEditableMesh::SplitEdge( const FEdgeID EdgeID, const TArray<float>& Splits
 			EdgeToCreate.VertexID0 = NewVertexIDs[ NewEdgeNumber ];
 			EdgeToCreate.VertexID1 = ( NewEdgeNumber == ( NewEdgeCount - 1 ) ) ? OriginalEdgeFarVertexID : NewVertexIDs[ NewEdgeNumber + 1 ];
 
-			EdgeToCreate.ConnectedPolygons = AffectedPolygons;
+			EdgeToCreate.ConnectedPolygons = GetMeshDescription()->GetEdgeConnectedPolygons( EdgeID );
 
 			// Copy edge attributes over from original edge
-			EdgeToCreate.EdgeAttributes.Attributes = OriginalEdgeAttributes;
+			EdgeToCreate.EdgeAttributes = EdgeAttributeList;
 		}
 
 		static TArray< FEdgeID > NewEdgeIDs;
@@ -3030,71 +3092,32 @@ void UEditableMesh::SplitEdge( const FEdgeID EdgeID, const TArray<float>& Splits
 
 	// Update all affected polygons with their new vertices.  Also, we'll fill in polygon-specific vertex attributes (texture coordinates)
 	{
-		for( const FAffectedPolygonEdge& AffectedPolygonEdge : AffectedPolygonEdges )
+		for( const FPolygonID ConnectedPolygonID : GetMeshDescription()->GetEdgeConnectedPolygons( EdgeID ) )
 		{
-			const FPolygonID PolygonID = AffectedPolygonEdge.PolygonID;
-
-			// Figure out where to start inserting vertices into this polygon, as well as which order to insert them.
-			bool bWindsForward = AffectedPolygonEdge.PolygonVertexNumbers[ 1 ] > AffectedPolygonEdge.PolygonVertexNumbers[ 0 ];
-			const int32 LargerIndex = bWindsForward ? AffectedPolygonEdge.PolygonVertexNumbers[ 1 ] : AffectedPolygonEdge.PolygonVertexNumbers[ 0 ];
-			const bool bIsContiguous = FMath::Abs( (int32)AffectedPolygonEdge.PolygonVertexNumbers[ 1 ] - (int32)AffectedPolygonEdge.PolygonVertexNumbers[ 0 ] ) == 1;
-			if( !bIsContiguous )
-			{
-				bWindsForward = !bWindsForward;
-			}
-
-			const int32 InsertAtIndex = bIsContiguous ? LargerIndex : ( LargerIndex + 1 );
-
-			// Figure out what the polygon's NEW vertex count (after we insert new vertices) will be
-			const int32 NewPolygonPerimeterVertexCount = GetPolygonPerimeterVertexCount( PolygonID ) + NewVertexPositions.Num();
-
 			static TArray<FVertexAndAttributes> VerticesToInsert;
-			VerticesToInsert.Reset( NewVertexPositions.Num() );
-			VerticesToInsert.SetNum( NewVertexPositions.Num(), false );
-			for( int32 InsertVertexNumber = 0; InsertVertexNumber < NewVertexPositions.Num(); ++InsertVertexNumber )
-			{
-				FVertexAndAttributes& VertexToInsert = VerticesToInsert[ InsertVertexNumber ];
-
-				const int32 DirectionalVertexNumber = bWindsForward ? InsertVertexNumber : ( ( NewVertexPositions.Num() - 1 ) - InsertVertexNumber );
-				VertexToInsert.VertexID = NewVertexIDs[ InsertVertexNumber ];
-				VertexToInsert.VertexInstanceID = FVertexInstanceID::Invalid;
+			VerticesToInsert.Reset( NumNewVertices );
 
-				// Generate texture coordinates for this vertex by interpolating between the polygon vertex texture coordinates on this polygon edge
-				{
-					const float Split = SortedSplits[ DirectionalVertexNumber ];
-					for( int32 TextureCoordinateIndex = 0; TextureCoordinateIndex < TextureCoordinateCount; ++TextureCoordinateIndex )
-					{
-						const FVector4 NewPolygonVertexTextureCoordinate = FMath::Lerp(
-							GetPolygonPerimeterVertexAttribute( PolygonID, AffectedPolygonEdge.PolygonVertexNumbers[ 0 ], UEditableMeshAttribute::VertexTextureCoordinate(), TextureCoordinateIndex ),
-							GetPolygonPerimeterVertexAttribute( PolygonID, AffectedPolygonEdge.PolygonVertexNumbers[ 1 ], UEditableMeshAttribute::VertexTextureCoordinate(), TextureCoordinateIndex ),
-							Split );
-
-						VertexToInsert.PolygonVertexAttributes.Attributes.Add( FMeshElementAttributeData( UEditableMeshAttribute::VertexTextureCoordinate(), TextureCoordinateIndex, NewPolygonVertexTextureCoordinate ) );
-					}
-				}
+			const FAffectedPolygonEdge& AffectedPolygonEdge = PolygonToAffectedEdge.FindRef( ConnectedPolygonID );
 
-				// Generate vertex color for this vertex by interpolating between the vertex colors on the edge
-				{
-					const float Split = SortedSplits[ DirectionalVertexNumber ];
-					const FVector4 NewVertexColor = FMath::Lerp(
-						GetPolygonPerimeterVertexAttribute( PolygonID, AffectedPolygonEdge.PolygonVertexNumbers[ 0 ], UEditableMeshAttribute::VertexColor(), 0 ),
-						GetPolygonPerimeterVertexAttribute( PolygonID, AffectedPolygonEdge.PolygonVertexNumbers[ 1 ], UEditableMeshAttribute::VertexColor(), 0 ),
-						Split );
+			for( int32 Count = 0; Count < NumNewVertices; ++Count )
+			{
+				const int32 VertexIndex = AffectedPolygonEdge.bReversed ? NumNewVertices - 1 - Count : Count;
 
-					VertexToInsert.PolygonVertexAttributes.Attributes.Add( FMeshElementAttributeData( UEditableMeshAttribute::VertexColor(), 0, NewVertexColor ) );
-				}
+				VerticesToInsert.Emplace();
+				FVertexAndAttributes& VertexAndAttributes = VerticesToInsert.Last();
+				VertexAndAttributes.VertexInstanceID = NewVertexInstanceIDs[ AffectedPolygonEdge.NewVertexInstanceArrayIndex + VertexIndex ];
 			}
 
-			// Add the new vertices to the polygon, and set texture coordinates for all of the affected polygon vertices
-			InsertPolygonPerimeterVertices( PolygonID, InsertAtIndex, VerticesToInsert );
+			// Add the new vertices to the polygon
+			InsertPolygonPerimeterVertices( ConnectedPolygonID, AffectedPolygonEdge.PerimeterVertexIndex, VerticesToInsert );
 		}
 	}
 
 	// Generate normals and tangents
-	PolygonsPendingNewTangentBasis.Append( AffectedPolygons );
+	PolygonsPendingNewTangentBasis.Append( GetMeshDescription()->GetEdgeConnectedPolygons( EdgeID ) );
 
 	// Retriangulate all of the affected polygons
-	PolygonsPendingTriangulation.Append( AffectedPolygons );
+	PolygonsPendingTriangulation.Append( GetMeshDescription()->GetEdgeConnectedPolygons( EdgeID ) );
 
 	UE_LOG( LogEditableMesh, Verbose, TEXT( "* SplitEdge returned %s" ), *LogHelpers::ArrayToString( OutNewVertexIDs ) );
 }
@@ -3125,6 +3148,9 @@ void UEditableMesh::FindPolygonLoop( const FEdgeID EdgeID, TArray<FEdgeID>& OutE
 	bool bCurrentEdgeIsBorderEdge = bStartedOnBorderEdge;
 	bool bCurrentEdgeIsInOppositeDirection = false;
 	bool bCurrentEdgeIsInOppositeDirectionFromStartEdge = false;
+
+	const TVertexAttributeArray<FVector>& VertexPositions = GetMeshDescription()->VertexAttributes().GetAttributes<FVector>( MeshAttribute::Vertex::Position );
+
 	for( ; ; )
 	{
 		// Add the current edge!
@@ -3138,8 +3164,8 @@ void UEditableMesh::FindPolygonLoop( const FEdgeID EdgeID, TArray<FEdgeID>& OutE
 		FVertexID CurrentEdgeVertexIDs[ 2 ];
 		GetEdgeVertices( CurrentEdgeID, /* Out */ CurrentEdgeVertexIDs[ 0 ], /* Out */ CurrentEdgeVertexIDs[ 1 ] );
 
-		const FVector CurrentEdgeVertex0 = GetVertexAttribute( CurrentEdgeVertexIDs[ 0 ], UEditableMeshAttribute::VertexPosition(), 0 );
-		const FVector CurrentEdgeVertex1 = GetVertexAttribute( CurrentEdgeVertexIDs[ 1 ], UEditableMeshAttribute::VertexPosition(), 0 );
+		const FVector CurrentEdgeVertex0 = VertexPositions[ CurrentEdgeVertexIDs[ 0 ] ];
+		const FVector CurrentEdgeVertex1 = VertexPositions[ CurrentEdgeVertexIDs[ 1 ] ];
 		const FVector CurrentEdgeDirection = ( CurrentEdgeVertex1 - CurrentEdgeVertex0 ).GetSafeNormal();
 
 		const FEdgeID NextEdgeIDInPath = OutReversedEdgeIDPathToTake.Num() > 0 ? OutReversedEdgeIDPathToTake.Pop( false ) : FEdgeID::Invalid;
@@ -3177,8 +3203,8 @@ void UEditableMesh::FindPolygonLoop( const FEdgeID EdgeID, TArray<FEdgeID>& OutE
 
 							const bool bIsBorderEdge = ( GetEdgeConnectedPolygonCount( CandidateEdgeID ) == 1 );
 
-							const FVector CandidateEdgeVertex0 = GetVertexAttribute( CandidateEdgeVertexIDs[ 0 ], UEditableMeshAttribute::VertexPosition(), 0 );
-							const FVector CandidateEdgeVertex1 = GetVertexAttribute( CandidateEdgeVertexIDs[ 1 ], UEditableMeshAttribute::VertexPosition(), 0 );
+							const FVector CandidateEdgeVertex0 = VertexPositions[ CandidateEdgeVertexIDs[ 0 ] ];
+							const FVector CandidateEdgeVertex1 = VertexPositions[ CandidateEdgeVertexIDs[ 1 ] ];
 							const FVector CandidateEdgeDirection = ( CandidateEdgeVertex1 - CandidateEdgeVertex0 ).GetSafeNormal();
 
 							const float DotProduct = FVector::DotProduct( CurrentEdgeDirection, CandidateEdgeDirection );
@@ -3194,8 +3220,8 @@ void UEditableMesh::FindPolygonLoop( const FEdgeID EdgeID, TArray<FEdgeID>& OutE
 									FVertexID LastSplitEdgeVertexIDs[ 2 ];
 									GetEdgeVertices( OutEdgeLoopEdgeIDs[ OutEdgeLoopEdgeIDs.Num() - 2 ], /* Out */ LastSplitEdgeVertexIDs[ 0 ], /* Out */ LastSplitEdgeVertexIDs[ 1 ] );
 
-									const FVector LastSplitEdgeVertex0 = GetVertexAttribute( LastSplitEdgeVertexIDs[ 0 ], UEditableMeshAttribute::VertexPosition(), 0 );
-									const FVector LastSplitEdgeVertex1 = GetVertexAttribute( LastSplitEdgeVertexIDs[ 1 ], UEditableMeshAttribute::VertexPosition(), 0 );
+									const FVector LastSplitEdgeVertex0 = VertexPositions[ LastSplitEdgeVertexIDs[ 0 ] ];
+									const FVector LastSplitEdgeVertex1 = VertexPositions[ LastSplitEdgeVertexIDs[ 1 ] ];
 
 									const FVector DirectionTowardCenterOfCurrentEdge =
 										( FMath::Lerp( CurrentEdgeVertex0, CurrentEdgeVertex1, 0.5f ) -
@@ -3211,8 +3237,8 @@ void UEditableMesh::FindPolygonLoop( const FEdgeID EdgeID, TArray<FEdgeID>& OutE
 									FVertexID BestEdgeVertexIDs[ 2 ];
 									GetEdgeVertices( BestEdgeID, /* Out */ BestEdgeVertexIDs[ 0 ], /* Out */ BestEdgeVertexIDs[ 1 ] );
 
-									const FVector BestEdgeVertex0 = GetVertexAttribute( BestEdgeVertexIDs[ 0 ], UEditableMeshAttribute::VertexPosition(), 0 );
-									const FVector BestEdgeVertex1 = GetVertexAttribute( BestEdgeVertexIDs[ 1 ], UEditableMeshAttribute::VertexPosition(), 0 );
+									const FVector BestEdgeVertex0 = VertexPositions[ BestEdgeVertexIDs[ 0 ] ];
+									const FVector BestEdgeVertex1 = VertexPositions[ BestEdgeVertexIDs[ 1 ] ];
 
 									const FVector DirectionTowardCenterOfBestEdge =
 										( FMath::Lerp( BestEdgeVertex0, BestEdgeVertex1, 0.5f ) -
@@ -3234,15 +3260,15 @@ void UEditableMesh::FindPolygonLoop( const FEdgeID EdgeID, TArray<FEdgeID>& OutE
 									// Edge directions are the same, but this is the very first split so we don't have a "flow" direction yet.
 									// Go ahead and prefer the edge that is closer to the initial edge.  This helps in the (uncommon) case of 
 									// multiple colinear edges on the same polygon (such as after SplitEdge() is called to insert a vertex on a polygon.)
-									const float BestEdgeDistance = [this, BestEdgeID, CurrentEdgeVertex0, CurrentEdgeVertex1]() -> float
+									const float BestEdgeDistance = [this, BestEdgeID, CurrentEdgeVertex0, CurrentEdgeVertex1, &VertexPositions]() -> float
 									{
 										check( BestEdgeID != FEdgeID::Invalid );
 
 										FVertexID BestEdgeVertexIDs[ 2 ];
 										GetEdgeVertices( BestEdgeID, /* Out */ BestEdgeVertexIDs[ 0 ], /* Out */ BestEdgeVertexIDs[ 1 ] );
 
-										const FVector BestEdgeVertex0 = GetVertexAttribute( BestEdgeVertexIDs[ 0 ], UEditableMeshAttribute::VertexPosition(), 0 );
-										const FVector BestEdgeVertex1 = GetVertexAttribute( BestEdgeVertexIDs[ 1 ], UEditableMeshAttribute::VertexPosition(), 0 );
+										const FVector BestEdgeVertex0 = VertexPositions[ BestEdgeVertexIDs[ 0 ] ];
+										const FVector BestEdgeVertex1 = VertexPositions[ BestEdgeVertexIDs[ 1 ] ];
 
 										FVector ClosestPoint0, ClosestPoint1;
 										FMath::SegmentDistToSegmentSafe( CurrentEdgeVertex0, CurrentEdgeVertex1, BestEdgeVertex0, BestEdgeVertex1, /* Out */ ClosestPoint0, /* Out */ ClosestPoint1 );
@@ -3727,7 +3753,7 @@ void UEditableMesh::DeleteOrphanVertices( const TArray<FVertexID>& VertexIDsToDe
 		{
 			const FVertexID VertexID = VertexIDsToDelete[ VertexNumber ];
 
-			// Make sure the vertex is truly an orphan.  We're not going to be able to restore it's polygon vertex attributes,
+			// Make sure the vertex is truly an orphan.  We're not going to be able to restore its polygon vertex attributes,
 			// because the polygons won't exist when we're restoring the change
 			const FMeshVertex& Vertex = Vertices[ VertexID ];
 			check( Vertex.ConnectedEdgeIDs.Num() == 0 );
@@ -3738,18 +3764,7 @@ void UEditableMesh::DeleteOrphanVertices( const TArray<FVertexID>& VertexIDsToDe
 
 			VertexToCreate.OriginalVertexID = VertexID;
 
-			for( const FName AttributeName : UEditableMesh::GetValidVertexAttributes() )
-			{
-				const int32 MaxAttributeIndex = GetMaxAttributeIndex( AttributeName );
-				for( int32 AttributeIndex = 0; AttributeIndex < MaxAttributeIndex; ++AttributeIndex )
-				{
-					VertexToCreate.VertexAttributes.Attributes.Emplace(
-						AttributeName,
-						AttributeIndex,
-						GetVertexAttribute( VertexID, AttributeName, AttributeIndex )
-					);
-				}
-			}
+			BackupAllAttributes( VertexToCreate.VertexAttributes, GetMeshDescription()->VertexAttributes(), VertexID );
 		}
 
 		AddUndo( MakeUnique<FCreateVerticesChange>( MoveTemp( RevertInput ) ) );
@@ -3765,7 +3780,7 @@ void UEditableMesh::DeleteOrphanVertices( const TArray<FVertexID>& VertexIDsToDe
 	{
 		for( const FVertexID VertexIDToDelete : VertexIDsToDelete )
 		{
-			Vertices.Remove( VertexIDToDelete );
+			GetMeshDescription()->DeleteVertex( VertexIDToDelete );
 		}
 	}
 
@@ -3798,18 +3813,7 @@ void UEditableMesh::DeleteVertexInstances( const TArray<FVertexInstanceID>& Vert
 			VertexInstanceToCreate.VertexID = VertexInstance.VertexID;
 			VertexInstanceToCreate.OriginalVertexInstanceID = VertexInstanceID;
 
-			for( const FName AttributeName : GetValidVertexInstanceAttributes() )
-			{
-				const int32 MaxAttributeIndex = GetMaxAttributeIndex( AttributeName );
-				for( int32 AttributeIndex = 0; AttributeIndex < MaxAttributeIndex; ++AttributeIndex )
-				{
-					VertexInstanceToCreate.VertexInstanceAttributes.Attributes.Emplace(
-						AttributeName,
-						AttributeIndex,
-						GetVertexInstanceAttribute( VertexInstanceID, AttributeName, AttributeIndex )
-					);
-				}
-			}
+			BackupAllAttributes( VertexInstanceToCreate.VertexInstanceAttributes, GetMeshDescription()->VertexInstanceAttributes(), VertexInstanceID );
 		}
 
 		AddUndo( MakeUnique<FCreateVertexInstancesChange>( MoveTemp( RevertInput ) ) );
@@ -3828,24 +3832,7 @@ void UEditableMesh::DeleteVertexInstances( const TArray<FVertexInstanceID>& Vert
 	{
 		for( const FVertexInstanceID VertexInstanceIDToDelete : VertexInstanceIDsToDelete )
 		{
-			// Check that no polygon is referencing this vertex instance
-			check( GetVertexInstanceConnectedPolygonCount( VertexInstanceIDToDelete ) == 0 );
-
-			// Remove the reference in the vertex it instances
-			const FMeshVertexInstance& VertexInstance = VertexInstances[ VertexInstanceIDToDelete ];
-			FMeshVertex& Vertex = Vertices[ VertexInstance.VertexID ];
-			verify( Vertex.VertexInstanceIDs.Remove( VertexInstanceIDToDelete ) == 1 );
-
-			if( bDeleteOrphanedVertices )
-			{
-				if( Vertex.VertexInstanceIDs.Num() == 0 )
-				{
-					OrphanedVertexIDs.Add( VertexInstance.VertexID );
-				}
-			}
-
-			// Delete the element from the container
-			VertexInstances.Remove( VertexInstanceIDToDelete );
+			GetMeshDescription()->DeleteVertexInstance( VertexInstanceIDToDelete, bDeleteOrphanedVertices ? &OrphanedVertexIDs : nullptr );
 		}
 	}
 
@@ -3884,18 +3871,7 @@ void UEditableMesh::DeleteEdges( const TArray<FEdgeID>& EdgeIDsToDelete, const b
 			EdgeToCreate.VertexID1 = Edge.VertexIDs[ 1 ];
 			EdgeToCreate.ConnectedPolygons = Edge.ConnectedPolygons;
 			
-			for( const FName AttributeName : UEditableMesh::GetValidEdgeAttributes() )
-			{
-				const int32 MaxAttributeIndex = GetMaxAttributeIndex( AttributeName );
-				for( int32 AttributeIndex = 0; AttributeIndex < MaxAttributeIndex; ++AttributeIndex )
-				{
-					EdgeToCreate.EdgeAttributes.Attributes.Emplace(
-						AttributeName,
-						AttributeIndex,
-						GetEdgeAttribute( EdgeID, AttributeName, AttributeIndex )
-					);
-				}
-			}
+			BackupAllAttributes( EdgeToCreate.EdgeAttributes, GetMeshDescription()->EdgeAttributes(), EdgeID );
 		}
 
 		AddUndo( MakeUnique<FCreateEdgesChange>( MoveTemp( RevertInput ) ) );
@@ -3913,24 +3889,7 @@ void UEditableMesh::DeleteEdges( const TArray<FEdgeID>& EdgeIDsToDelete, const b
 
 		for( const FEdgeID EdgeID : EdgeIDsToDelete )
 		{
-			FMeshEdge& Edge = Edges[ EdgeID ];
-
-			for( const FVertexID EdgeVertexID : Edge.VertexIDs )
-			{
-				FMeshVertex& Vertex = Vertices[ EdgeVertexID ];
-				verify( Vertex.ConnectedEdgeIDs.RemoveSingle( EdgeID ) == 1 );
-
-				// If the vertex has no more edges connected, we'll keep track of that so we can delete the vertex later
-				if( bDeleteOrphanedVertices && Vertex.ConnectedEdgeIDs.Num() == 0 )
-				{
-					check( Vertex.VertexInstanceIDs.Num() == 0 );	// We must already have deleted any vertex instances
-					check( !OrphanedVertexIDs.Contains( EdgeVertexID ) );	// Orphaned vertex shouldn't have already been orphaned by an earlier deleted edge passed into this function
-					OrphanedVertexIDs.Add( EdgeVertexID );
-				}
-			}
-
-			// Delete the edge
-			Edges.Remove( EdgeID );
+			GetMeshDescription()->DeleteEdge( EdgeID, bDeleteOrphanedVertices ? &OrphanedVertexIDs : nullptr );
 		}
 
 		// If we orphaned any vertices and we were asked to delete those, then we'll go ahead and do that now.
@@ -3948,24 +3907,15 @@ void UEditableMesh::CreateEmptyVertexRange( const int32 NumVerticesToCreate, TAr
 {
 	UE_LOG( LogEditableMesh, Verbose, TEXT( "CreateEmptyVertexRange: %d" ), NumVerticesToCreate );
 
-	FVertexArray& Vertices = GetMeshDescription()->Vertices();
-
 	OutNewVertexIDs.Reset( NumVerticesToCreate );
 
-	// Reserve elements
-	Vertices.Reserve( Vertices.Num() + NumVerticesToCreate );
-
-	for( int32 Count = 0; Count < NumVerticesToCreate; ++Count )
+	// Create vertices
 	{
-		// Allocate vertex
-		const FVertexID VertexID = Vertices.Add();
-		FMeshVertex& Vertex = Vertices[ VertexID ];
-
-		// Initialize vertex
-		Vertex.VertexPosition = FVector::ZeroVector;
-		Vertex.CornerSharpness = 0.0f;
-
-		OutNewVertexIDs.Add( VertexID );
+		GetMeshDescription()->ReserveNewVertices( NumVerticesToCreate );
+		for( int32 Count = 0; Count < NumVerticesToCreate; ++Count )
+		{
+			OutNewVertexIDs.Add( GetMeshDescription()->CreateVertex() );
+		}
 	}
 
 	// NOTE: We iterate backwards, to delete vertices in the opposite order that we added them
@@ -3994,33 +3944,25 @@ void UEditableMesh::CreateVertices( const TArray<FVertexToCreate>& VerticesToCre
 {
 	UE_LOG( LogEditableMesh, Verbose, TEXT( "CreateVertices: %s" ), *LogHelpers::ArrayToString( VerticesToCreate ) );
 
-	FVertexArray& Vertices = GetMeshDescription()->Vertices();
-
-	OutNewVertexIDs.Reset( VerticesToCreate.Num() );
-
-	// Reserve elements
-	Vertices.Reserve( Vertices.Num() + VerticesToCreate.Num() );
-
-	for( const FVertexToCreate& VertexToCreate : VerticesToCreate )
+	// Create vertices
 	{
-		// Allocate vertex
-		FVertexID VertexID = VertexToCreate.OriginalVertexID;
-		if( VertexID != FVertexID::Invalid )
-		{
-			Vertices.Insert( VertexID );
-		}
-		else
-		{
-			VertexID = Vertices.Add();
-		}
-
-		FMeshVertex& Vertex = Vertices[ VertexID ];
+		OutNewVertexIDs.Reset( VerticesToCreate.Num() );
 
-		// Initialize vertex
-		Vertex.VertexPosition = FVector::ZeroVector;
-		Vertex.CornerSharpness = 0.0f;
+		GetMeshDescription()->ReserveNewVertices( VerticesToCreate.Num() );
+		for( const FVertexToCreate& VertexToCreate : VerticesToCreate )
+		{
+			FVertexID VertexID = VertexToCreate.OriginalVertexID;
+			if( VertexID != FVertexID::Invalid )
+			{
+				GetMeshDescription()->CreateVertexWithID( VertexID );
+			}
+			else
+			{
+				VertexID = GetMeshDescription()->CreateVertex();
+			}
 
-		OutNewVertexIDs.Add( VertexID );
+			OutNewVertexIDs.Add( VertexID );
+		}
 	}
 
 	// NOTE: We iterate backwards, to delete vertices in the opposite order that we added them
@@ -4046,7 +3988,7 @@ void UEditableMesh::CreateVertices( const TArray<FVertexToCreate>& VerticesToCre
 	{
 		for( const FMeshElementAttributeData& VertexAttribute : VerticesToCreate[ Index ].VertexAttributes.Attributes )
 		{
-			SetVertexAttribute( OutNewVertexIDs[ Index ], VertexAttribute.AttributeName, VertexAttribute.AttributeIndex, VertexAttribute.AttributeValue );
+			SetVertexAttribute( OutNewVertexIDs[ Index ], VertexAttribute );
 		}
 	}
 
@@ -4058,43 +4000,25 @@ void UEditableMesh::CreateVertexInstances( const TArray<FVertexInstanceToCreate>
 {
 	UE_LOG( LogEditableMesh, Verbose, TEXT( "CreateVertexInstances: %s" ), *LogHelpers::ArrayToString( VertexInstancesToCreate ) );
 
-	FVertexArray& Vertices = GetMeshDescription()->Vertices();
-	FVertexInstanceArray& VertexInstances = GetMeshDescription()->VertexInstances();
-
-	OutNewVertexInstanceIDs.Reset( VertexInstancesToCreate.Num() );
-
-	// Reserve elements
-	VertexInstances.Reserve( VertexInstances.Num() + VertexInstancesToCreate.Num() );
-
-	for( const FVertexInstanceToCreate& VertexInstanceToCreate : VertexInstancesToCreate )
+	// Create new vertex instances
 	{
-		// Allocate vertex instance
-		FVertexInstanceID VertexInstanceID = VertexInstanceToCreate.OriginalVertexInstanceID;
-		if( VertexInstanceID != FVertexInstanceID::Invalid )
-		{
-			VertexInstances.Insert( VertexInstanceID );
-		}
-		else
-		{
-			VertexInstanceID = VertexInstances.Add();
-		}
-
-		FMeshVertexInstance& VertexInstance = VertexInstances[ VertexInstanceID ];
+		OutNewVertexInstanceIDs.Reset( VertexInstancesToCreate.Num() );
 
-		// Initialize vertex instance
-		VertexInstance.VertexID = VertexInstanceToCreate.VertexID;
-
-		VertexInstance.Normal = FVector::ZeroVector;
-		VertexInstance.Tangent = FVector::ZeroVector;
-		VertexInstance.BinormalSign = 0.0f;
-		VertexInstance.Color = FLinearColor::White;
-		VertexInstance.VertexUVs.SetNumZeroed( TextureCoordinateCount );
-
-		// Add to the list of instances in the vertex
-		check( !Vertices[ VertexInstance.VertexID ].VertexInstanceIDs.Contains( VertexInstanceID ) );
-		Vertices[ VertexInstance.VertexID ].VertexInstanceIDs.Add( VertexInstanceID );
+		GetMeshDescription()->ReserveNewVertexInstances( VertexInstancesToCreate.Num() );
+		for( const FVertexInstanceToCreate& VertexInstanceToCreate : VertexInstancesToCreate )
+		{
+			FVertexInstanceID VertexInstanceID = VertexInstanceToCreate.OriginalVertexInstanceID;
+			if( VertexInstanceID != FVertexInstanceID::Invalid )
+			{
+				GetMeshDescription()->CreateVertexInstanceWithID( VertexInstanceID, VertexInstanceToCreate.VertexID );
+			}
+			else
+			{
+				VertexInstanceID = GetMeshDescription()->CreateVertexInstance( VertexInstanceToCreate.VertexID );
+			}
 
-		OutNewVertexInstanceIDs.Add( VertexInstanceID );
+			OutNewVertexInstanceIDs.Add( VertexInstanceID );
+		}
 	}
 
 	// NOTE: We iterate backwards, to delete vertex instances in the opposite order that we added them
@@ -4121,7 +4045,7 @@ void UEditableMesh::CreateVertexInstances( const TArray<FVertexInstanceToCreate>
 	{
 		for( const FMeshElementAttributeData& VertexInstanceAttribute : VertexInstancesToCreate[ Index ].VertexInstanceAttributes.Attributes )
 		{
-			SetVertexInstanceAttribute( OutNewVertexInstanceIDs[ Index ], VertexInstanceAttribute.AttributeName, VertexInstanceAttribute.AttributeIndex, VertexInstanceAttribute.AttributeValue );
+			SetVertexInstanceAttribute( OutNewVertexInstanceIDs[ Index ], VertexInstanceAttribute );
 		}
 	}
 
@@ -4133,42 +4057,25 @@ void UEditableMesh::CreateEdges( const TArray<FEdgeToCreate>& EdgesToCreate, TAr
 {
 	UE_LOG( LogEditableMesh, Verbose, TEXT( "CreateEdges: %s" ), *LogHelpers::ArrayToString( EdgesToCreate ) );
 
-	FVertexArray& Vertices = GetMeshDescription()->Vertices();
-	FEdgeArray& Edges = GetMeshDescription()->Edges();
-
-	OutNewEdgeIDs.Reset( EdgesToCreate.Num() );
-
-	// Reserve elements
-	Edges.Reserve( Edges.Num() + EdgesToCreate.Num() );
-
-	for( const FEdgeToCreate& EdgeToCreate : EdgesToCreate )
+	// Create new edges in the mesh description
 	{
-		// Allocate edge
-		FEdgeID EdgeID = EdgeToCreate.OriginalEdgeID;
-		if( EdgeID != FEdgeID::Invalid )
-		{
-			Edges.Insert( EdgeID );
-		}
-		else
-		{
-			EdgeID = Edges.Add();
-		}
-
-		FMeshEdge& Edge = Edges[ EdgeID ];
-
-		// Initialize edge
-		Edge.VertexIDs[ 0 ] = EdgeToCreate.VertexID0;
-		Edge.VertexIDs[ 1 ] = EdgeToCreate.VertexID1;
-		Edge.ConnectedPolygons = EdgeToCreate.ConnectedPolygons;
+		OutNewEdgeIDs.Reset( EdgesToCreate.Num() );
 
-		Edge.bIsHardEdge = false;
-		Edge.CreaseSharpness = 0.0f;
-
-		// Connect the edge to its vertices
-		Vertices[ Edge.VertexIDs[ 0 ] ].ConnectedEdgeIDs.AddUnique( EdgeID );
-		Vertices[ Edge.VertexIDs[ 1 ] ].ConnectedEdgeIDs.AddUnique( EdgeID );
+		GetMeshDescription()->ReserveNewEdges( EdgesToCreate.Num() );
+		for( const FEdgeToCreate& EdgeToCreate : EdgesToCreate )
+		{
+			FEdgeID EdgeID = EdgeToCreate.OriginalEdgeID;
+			if( EdgeID != FEdgeID::Invalid )
+			{
+				GetMeshDescription()->CreateEdgeWithID( EdgeID, EdgeToCreate.VertexID0, EdgeToCreate.VertexID1, EdgeToCreate.ConnectedPolygons );
+			}
+			else
+			{
+				EdgeID = GetMeshDescription()->CreateEdge( EdgeToCreate.VertexID0, EdgeToCreate.VertexID1, EdgeToCreate.ConnectedPolygons );
+			}
 
-		OutNewEdgeIDs.Add( EdgeID );
+			OutNewEdgeIDs.Add( EdgeID );
+		}
 	}
 
 	// NOTE: We iterate backwards, to delete edges in the opposite order that we added them
@@ -4195,7 +4102,7 @@ void UEditableMesh::CreateEdges( const TArray<FEdgeToCreate>& EdgesToCreate, TAr
 		for( const FMeshElementAttributeData& EdgeAttribute : EdgesToCreate[ Index ].EdgeAttributes.Attributes )
 		{
 			const bool bIsUndo = false;
-			SetEdgeAttribute( OutNewEdgeIDs[ Index ], EdgeAttribute.AttributeName, EdgeAttribute.AttributeIndex, EdgeAttribute.AttributeValue, bIsUndo );
+			SetEdgeAttribute( OutNewEdgeIDs[ Index ], EdgeAttribute, bIsUndo );
 		}
 	}
 
@@ -4207,59 +4114,27 @@ FVertexInstanceID UEditableMesh::CreateVertexInstanceForContourVertex( const FVe
 {
 	FVertexInstanceID NewVertexInstanceID = ContourVertex.VertexInstanceID;
 
-	// Determine whether we need to create a new vertex instance or not. We need to if:
-	//   a) no valid vertex instance ID was supplied;
-	// and then:
-	//   b) if there are no vertex instances in the same soft edged polygon group as the new polygon which share the attributes being assigned
 	if( NewVertexInstanceID == FVertexInstanceID::Invalid )
 	{
 		const FVertexID VertexID = ContourVertex.VertexID;
 		check( VertexID != FVertexID::Invalid );
 
-		// Find candidate vertex instances in the same smoothing group as this polygon
-		static TArray<FVertexInstanceID> PossibleVertexInstances;
-		PossibleVertexInstances.Reset();
+		// For now, always create a new vertex instance.
+		// @todo mesheditor: revisit splitting/welding identical vertex instances.
+		// This should be done as a post-step at the end of the transaction, according to adjacent edge hardness.
 
-		const bool bPolygonNotYetInitialized = true;
-		GetVertexInstancesInSameSoftEdgedGroup( VertexID, PolygonID, bPolygonNotYetInitialized, PossibleVertexInstances );
+		FVertexInstanceToCreate VertexInstanceToCreate;
+		VertexInstanceToCreate.VertexID = VertexID;
+		VertexInstanceToCreate.VertexInstanceAttributes = ContourVertex.PolygonVertexAttributes;
 
-		// Iterate through each potential vertex instance, checking if its attributes are equal before and after the new attributes are applied.
-		// If so, we can merge with this vertex.
-		// Note: unspecified attributes in the FVertexAndAttributes struct are "don't cares".
-		for( const FVertexInstanceID PossibleVertexInstance : PossibleVertexInstances )
-		{
-			// Get mergeable attribute data for this vertex instance
-			static TArray<FMeshElementAttributeData> AllAttributes;
-			GetMergedNamedVertexInstanceAttributeData( GetMergeableVertexInstanceAttributes(), PossibleVertexInstance, TArray<FMeshElementAttributeData>(), AllAttributes );
-
-			// Get attribute data for this vertex instance as it would be after applying the new attributes
-			static TArray<FMeshElementAttributeData> AllAttributesModified;
-			GetMergedNamedVertexInstanceAttributeData( GetMergeableVertexInstanceAttributes(), PossibleVertexInstance, ContourVertex.PolygonVertexAttributes.Attributes, AllAttributesModified );
+		static TArray<FVertexInstanceToCreate> VertexInstancesToCreate;
+		VertexInstancesToCreate.Reset();
+		VertexInstancesToCreate.Add( VertexInstanceToCreate );
 
-			// If it's the same, we can use this vertex instance instead of creating a new one
-			if( AreAttributeListsEqual( AllAttributes, AllAttributesModified ) )
-			{
-				NewVertexInstanceID = PossibleVertexInstance;
-				break;
-			}
-		}
-
-		// If we didn't find any suitable existing vertex instance to merge with, create a new one.
-		if( NewVertexInstanceID == FVertexInstanceID::Invalid )
-		{
-			FVertexInstanceToCreate VertexInstanceToCreate;
-			VertexInstanceToCreate.VertexID = VertexID;
-			VertexInstanceToCreate.VertexInstanceAttributes = ContourVertex.PolygonVertexAttributes;
-
-			static TArray<FVertexInstanceToCreate> VertexInstancesToCreate;
-			VertexInstancesToCreate.Reset();
-			VertexInstancesToCreate.Add( VertexInstanceToCreate );
-
-			static TArray<FVertexInstanceID> CreatedVertexInstances;
-			CreateVertexInstances( VertexInstancesToCreate, CreatedVertexInstances );
-			check( CreatedVertexInstances.Num() == 1 );
-			NewVertexInstanceID = CreatedVertexInstances[ 0 ];
-		}
+		static TArray<FVertexInstanceID> CreatedVertexInstances;
+		CreateVertexInstances( VertexInstancesToCreate, CreatedVertexInstances );
+		check( CreatedVertexInstances.Num() == 1 );
+		NewVertexInstanceID = CreatedVertexInstances[ 0 ];
 	}
 	else
 	{
@@ -4271,22 +4146,23 @@ FVertexInstanceID UEditableMesh::CreateVertexInstanceForContourVertex( const FVe
 }
 
 
-void UEditableMesh::CreatePolygonContour( const FPolygonID PolygonID, const TArray<FVertexAndAttributes>& Contour, const EPolygonEdgeHardness PolygonEdgeHardness, TArray<FEdgeID>& OutEdgeIDs, TArray<FVertexInstanceID>& OutVertexInstanceIDs )
+void UEditableMesh::CreatePolygonContour( const TArray<FVertexAndAttributes>& Contour, TArray<FEdgeID>& OutExistingEdgeIDs, TArray<FEdgeID>& OutNewEdgeIDs, TArray<UMeshDescription::FContourPoint>& OutContourPoints )
 {
-	FVertexInstanceArray& VertexInstances = GetMeshDescription()->VertexInstances();
-	FEdgeArray& Edges = GetMeshDescription()->Edges();
-
 	// All polygons must have at least three vertices
 	const int32 NumContourVertices = Contour.Num();
 	check( NumContourVertices >= 3 );
 
-	OutEdgeIDs.Reset( NumContourVertices );
-	OutVertexInstanceIDs.SetNumUninitialized( NumContourVertices );
+	OutExistingEdgeIDs.Reset( NumContourVertices );
+	OutNewEdgeIDs.Reset( NumContourVertices );
+	OutContourPoints.SetNumUninitialized( NumContourVertices );
 
-	static TArray<FEdgeID> AllEdges;
-	AllEdges.SetNumUninitialized( NumContourVertices );
+	static TArray<FEdgeToCreate> EdgesToCreate;
+	EdgesToCreate.Reset();
 
-	// Get or create edges
+	static TArray<FVertexInstanceToCreate> VertexInstancesToCreate;
+	VertexInstancesToCreate.Reset();
+
+	// Get edges or decide to create them
 	for( int32 VertexNumber = 0; VertexNumber < NumContourVertices; ++VertexNumber )
 	{
 		const int32 NextVertexNumber = ( VertexNumber + 1 ) % NumContourVertices;
@@ -4300,79 +4176,104 @@ void UEditableMesh::CreatePolygonContour( const FPolygonID PolygonID, const TArr
 		{
 			const FVertexInstanceID VertexInstanceID0 = Contour[ VertexNumber ].VertexInstanceID;
 			check( VertexInstanceID0 != FVertexInstanceID::Invalid );
-			VertexID0 = VertexInstances[ VertexInstanceID0 ].VertexID;
+			VertexID0 = GetVertexInstanceVertex( VertexInstanceID0 );
 		}
 
 		if( VertexID1 == FVertexID::Invalid )
 		{
 			const FVertexInstanceID VertexInstanceID1 = Contour[ NextVertexNumber ].VertexInstanceID;
 			check( VertexInstanceID1 != FVertexInstanceID::Invalid );
-			VertexID1 = VertexInstances[ VertexInstanceID1 ].VertexID;
+			VertexID1 = GetVertexInstanceVertex( VertexInstanceID1 );
 		}
 
-		bool bOutEdgeWindingIsReversed;
-		FEdgeID EdgeID = GetVertexPairEdge( VertexID0, VertexID1, bOutEdgeWindingIsReversed );
+		FEdgeID EdgeID = GetMeshDescription()->GetVertexPairEdge( VertexID0, VertexID1 );
+
+		// Copy the edge ID to the output.
+		// If it is invalid, it will be patched later when new edges are created.
+		OutContourPoints[ VertexNumber ].EdgeID = EdgeID;
+
 		if( EdgeID == FEdgeID::Invalid )
 		{
-			// Create the new edge!  Note that this does not connect the edge to the polygon.  We expect the caller to do that afterwards.
-			FEdgeToCreate EdgeToCreate;
+			// Create the new edge. Don't specify any attributes now.
+			EdgesToCreate.Emplace();
+			FEdgeToCreate& EdgeToCreate = EdgesToCreate.Last();
 			EdgeToCreate.VertexID0 = VertexID0;
 			EdgeToCreate.VertexID1 = VertexID1;
+		}
+		else
+		{
+			OutExistingEdgeIDs.Add( EdgeID );
+		}
+	}
 
-			if( PolygonEdgeHardness == EPolygonEdgeHardness::NewEdgesHard )
-			{
-				EdgeToCreate.EdgeAttributes.Attributes.Emplace( UEditableMeshAttribute::EdgeIsHard(), 0, FVector4( 1.0f ) );
-			}
-
-			static TArray<FEdgeToCreate> EdgesToCreate;
-			EdgesToCreate.Reset();
-			EdgesToCreate.Add( EdgeToCreate );
+	// Create missing edges
+	if( EdgesToCreate.Num() > 0 )
+	{
+		CreateEdges( EdgesToCreate, OutNewEdgeIDs );
 
-			static TArray<FEdgeID> NewEdgeIDs;
-			CreateEdges( EdgesToCreate, /* Out */ NewEdgeIDs );
-			check( NewEdgeIDs.Num() == 1 );
-			EdgeID = NewEdgeIDs[ 0 ];
+		// Fill in missing edge indices
+		int32 NewEdgeIndex = 0;
 
-			OutEdgeIDs.Add( EdgeID );
+		for( UMeshDescription::FContourPoint& ContourPoint : OutContourPoints )
+		{
+			if( ContourPoint.EdgeID == FEdgeID::Invalid )
+			{
+				ContourPoint.EdgeID = OutNewEdgeIDs[ NewEdgeIndex ];
+				NewEdgeIndex++;
+			}
 		}
-
-		// Connect the edge to the polygon
-		FMeshEdge& Edge = Edges[ EdgeID ];
-		Edge.ConnectedPolygons.Add( PolygonID );
-
-		AllEdges[ VertexNumber ] = EdgeID;
 	}
 
-	if( PolygonEdgeHardness == EPolygonEdgeHardness::AllEdgesHard ||
-		PolygonEdgeHardness == EPolygonEdgeHardness::AllEdgesSoft )
+	// Assign vertex instances to the polygon
+	for( int32 VertexNumber = 0; VertexNumber < NumContourVertices; ++VertexNumber )
 	{
-		static TArray<FAttributesForEdge> AttributesForEdges;
-		AttributesForEdges.Reset( AllEdges.Num() );
+		// Copy the supplied vertex instance ID into the output.
+		// If it is valid, there is nothing more to be done.
+		// If it is not valid, copy the invalid ID. They will be replaced by the IDs of the new vertex instances in a later pass.
+		OutContourPoints[ VertexNumber ].VertexInstanceID = Contour[ VertexNumber ].VertexInstanceID;
 
-		for( const FEdgeID EdgeID : AllEdges )
+		if( Contour[ VertexNumber ].VertexInstanceID != FVertexInstanceID::Invalid )
 		{
-			AttributesForEdges.Emplace();
-			FAttributesForEdge& AttributesForEdge = AttributesForEdges.Last();
-			AttributesForEdge.EdgeID = EdgeID;
-			AttributesForEdge.EdgeAttributes.Attributes.Emplace( UEditableMeshAttribute::EdgeIsHard(), 0, FVector4( PolygonEdgeHardness == EPolygonEdgeHardness::AllEdgesHard ? 1.0f : 0.0f ) );
+			// Cannot specify both a vertex instance ID and a vertex ID
+			check( Contour[ VertexNumber ].VertexID == FVertexID::Invalid );
+
+			// Vertex instance ID was specified; nothing more to do, it has already been output directly
+			// @todo editable mesh: THIS WILL BREAK if you specify a vertex instance which is being used by a polygon the other side of a hard edge.
+			// Decide what to do about this.
 		}
+		else
+		{
+			// We need to create a vertex instance; check that the vertex ID is valid
+			const FVertexID VertexID = Contour[ VertexNumber ].VertexID;
+			check( VertexID != FVertexID::Invalid );
 
-		SetEdgesAttributes( AttributesForEdges );
+			// For now, always create a new vertex instance. It will be merged with an existing one later if appropriate.
+			VertexInstancesToCreate.Emplace();
+			FVertexInstanceToCreate& VertexInstanceToCreate = VertexInstancesToCreate.Last();
+			VertexInstanceToCreate.VertexID = VertexID;
+			VertexInstanceToCreate.VertexInstanceAttributes = Contour[ VertexNumber ].PolygonVertexAttributes;
+		}
 	}
 
-	// Get or create new vertex instances for this polygon.
-	// Do this after creating all the edges because:
-	//  1) we can establish a smoothing group for the new polygon based on adjacent edge hardness.
-	//  2) the revert operation will first delete all vertex instances, and then the edges (which can only delete orphaned vertices if they have no instances).
-	for( int32 VertexNumber = 0; VertexNumber < NumContourVertices; ++VertexNumber )
+	// Create any vertex instances necessary
+	if( VertexInstancesToCreate.Num() > 0 )
 	{
-		const FVertexInstanceID ContourVertexInstanceID = CreateVertexInstanceForContourVertex( Contour[ VertexNumber ], PolygonID );
-		OutVertexInstanceIDs[ VertexNumber ] = ContourVertexInstanceID;
+		static TArray<FVertexInstanceID> NewVertexInstanceIDs;
+		NewVertexInstanceIDs.Reset( NumContourVertices );
 
-		// Connect this vertex instance to the polygon
-		FMeshVertexInstance& ContourVertexInstance = VertexInstances[ ContourVertexInstanceID ];
-		check( !ContourVertexInstance.ConnectedPolygons.Contains( PolygonID ) );
-		ContourVertexInstance.ConnectedPolygons.Add( PolygonID );
+		CreateVertexInstances( VertexInstancesToCreate, NewVertexInstanceIDs );
+
+		// Fill in missing vertex instance indices
+		int32 NewVertexInstanceIndex = 0;
+
+		for( UMeshDescription::FContourPoint& ContourPoint : OutContourPoints )
+		{
+			if( ContourPoint.VertexInstanceID == FEdgeID::Invalid )
+			{
+				ContourPoint.VertexInstanceID = NewVertexInstanceIDs[ NewVertexInstanceIndex ];
+				NewVertexInstanceIndex++;
+			}
+		}
 	}
 }
 
@@ -4381,52 +4282,83 @@ void UEditableMesh::CreatePolygons( const TArray<FPolygonToCreate>& PolygonsToCr
 {
 	UE_LOG( LogEditableMesh, Verbose, TEXT( "CreatePolygons: %s" ), *LogHelpers::ArrayToString( PolygonsToCreate ) );
 
-	FPolygonArray& Polygons = GetMeshDescription()->Polygons();
-	FPolygonGroupArray& PolygonGroups = GetMeshDescription()->PolygonGroups();
-
 	OutNewPolygonIDs.Reset( PolygonsToCreate.Num() );
 	OutNewEdgeIDs.Reset();
 
-	Polygons.Reserve( Polygons.Num() + PolygonsToCreate.Num() );
-
-	for( const FPolygonToCreate& PolygonToCreate : PolygonsToCreate )
+	// Actually create the polygons
 	{
-		// Allocate polygon
-		FPolygonID PolygonID = PolygonToCreate.OriginalPolygonID;
-		if( PolygonID != FPolygonID::Invalid )
-		{
-			Polygons.Insert( PolygonID );
-		}
-		else
+		static TArray<FAttributesForEdge> AttributesForEdges;
+		AttributesForEdges.Reset();
+
+		static TArray<FEdgeID> ExistingEdgeIDs;
+		ExistingEdgeIDs.Reset();
+
+		GetMeshDescription()->ReserveNewPolygons( PolygonsToCreate.Num() );
+		for( const FPolygonToCreate& PolygonToCreate : PolygonsToCreate )
 		{
-			PolygonID = Polygons.Add();
-		}
+			static TArray<UMeshDescription::FContourPoint> PerimeterContourPoints;
+			static TArray<FEdgeID> NewEdgeIDsForContour;
+			static TArray<FEdgeID> ExistingEdgeIDsForContour;
 
-		// Connect polygon to group and vice versa
-		FMeshPolygon& Polygon = Polygons[ PolygonID ];
-		FMeshPolygonGroup& PolygonGroup = PolygonGroups[ PolygonToCreate.PolygonGroupID ];
-		Polygon.PolygonGroupID = PolygonToCreate.PolygonGroupID;
-		PolygonGroup.Polygons.Add( PolygonID );
+			CreatePolygonContour( PolygonToCreate.PerimeterVertices, ExistingEdgeIDsForContour, NewEdgeIDsForContour, PerimeterContourPoints );
+			ExistingEdgeIDs.Append( ExistingEdgeIDsForContour );
+			OutNewEdgeIDs.Append( NewEdgeIDsForContour );
 
-		OutNewPolygonIDs.Add( PolygonID );
+			static TArray<TArray<UMeshDescription::FContourPoint>> HoleContourPoints;
 
-		// Create edges and vertex instances for perimeter contour
-		static TArray<FEdgeID> PerimeterEdgeIDs;
-		CreatePolygonContour( PolygonID, PolygonToCreate.PerimeterVertices, PolygonToCreate.PolygonEdgeHardness, PerimeterEdgeIDs, Polygon.PerimeterContour.VertexInstanceIDs );
-		OutNewEdgeIDs.Append( PerimeterEdgeIDs );
+			// Iterate through hole definitions, and create edges and vertex instances for those
+			for( const FPolygonHoleVertices& PolygonHole : PolygonToCreate.PolygonHoles )
+			{
+				HoleContourPoints.Emplace();
 
-		// Iterate through hole definitions, and create edges and vertex instances for those
-		for( const FPolygonHoleVertices& PolygonHole : PolygonToCreate.PolygonHoles )
-		{
-			Polygon.HoleContours.Emplace();
-			FMeshPolygonContour& Contour = Polygon.HoleContours.Last();
+				CreatePolygonContour( PolygonHole.HoleVertices, ExistingEdgeIDsForContour, NewEdgeIDsForContour, HoleContourPoints.Last() );
+				ExistingEdgeIDs.Append( ExistingEdgeIDsForContour );
+				OutNewEdgeIDs.Append( NewEdgeIDsForContour );
+			}
 
-			static TArray<FEdgeID> HoleEdgeIDs;
-			CreatePolygonContour( PolygonID, PolygonHole.HoleVertices, EPolygonEdgeHardness::AllEdgesHard, HoleEdgeIDs, Contour.VertexInstanceIDs );
-			OutNewEdgeIDs.Append( HoleEdgeIDs );
+			FPolygonID PolygonID = PolygonToCreate.OriginalPolygonID;
+			if( PolygonID != FPolygonID::Invalid )
+			{
+				GetMeshDescription()->CreatePolygonWithID( PolygonID, PolygonToCreate.PolygonGroupID, PerimeterContourPoints, HoleContourPoints );
+			}
+			else
+			{
+				PolygonID = GetMeshDescription()->CreatePolygon( PolygonToCreate.PolygonGroupID, PerimeterContourPoints, HoleContourPoints );
+			}
+
+			OutNewPolygonIDs.Add( PolygonID );
+
+			// Set edge hardness for polygon edges according to PolygonToCreate mode
+			AttributesForEdges.Reserve( AttributesForEdges.Num() + OutNewEdgeIDs.Num() + ExistingEdgeIDs.Num() );
+
+			const bool bHardEdge = ( PolygonToCreate.PolygonEdgeHardness == EPolygonEdgeHardness::AllEdgesHard || PolygonToCreate.PolygonEdgeHardness == EPolygonEdgeHardness::NewEdgesHard );
+			for( const FEdgeID EdgeID : OutNewEdgeIDs )
+			{
+				AttributesForEdges.Emplace();
+				FAttributesForEdge& AttributesForEdge = AttributesForEdges.Last();
+				AttributesForEdge.EdgeID = EdgeID;
+				AttributesForEdge.EdgeAttributes.Attributes.Emplace( MeshAttribute::Edge::IsHard, 0, FMeshElementAttributeValue( bHardEdge ) );
+			}
+
+			// If we are setting all edges' hardnesses (not just new ones), set existing edges' attributes here.
+			// This will also split any vertex instances which have just been included in a polygon if necessary.
+			if( PolygonToCreate.PolygonEdgeHardness == EPolygonEdgeHardness::AllEdgesHard || PolygonToCreate.PolygonEdgeHardness == EPolygonEdgeHardness::AllEdgesSoft )
+			{
+				for( const FEdgeID EdgeID : ExistingEdgeIDs )
+				{
+					AttributesForEdges.Emplace();
+					FAttributesForEdge& AttributesForEdge = AttributesForEdges.Last();
+					AttributesForEdge.EdgeID = EdgeID;
+					AttributesForEdge.EdgeAttributes.Attributes.Emplace( MeshAttribute::Edge::IsHard, 0, FMeshElementAttributeValue( bHardEdge ) );
+				}
+			}
 		}
+
+		SetEdgesAttributes( AttributesForEdges );
 	}
 
+	// @todo editable mesh: merge vertex instances if possible here
+
 	for( UEditableMeshAdapter* Adapter : Adapters )
 	{
 		Adapter->OnCreatePolygons( this, OutNewPolygonIDs );
@@ -4483,43 +4415,6 @@ void UEditableMesh::BackupPolygonContour( const FMeshPolygonContour& Contour, TA
 }
 
 
-void UEditableMesh::DeletePolygonContour( const FPolygonID PolygonID, const FMeshPolygonContour& Contour, const bool bDeleteOrphanedEdges, const bool bDeleteOrphanedVertexInstances, TArray<FEdgeID>& OrphanedEdgeIDs, TArray<FVertexInstanceID>& OrphanedVertexInstanceIDs )
-{
-	FVertexInstanceArray& VertexInstances = GetMeshDescription()->VertexInstances();
-
-	for( const FVertexInstanceID VertexInstanceID : Contour.VertexInstanceIDs )
-	{
-		// Disconnect the polygon from the vertex instance
-		FMeshVertexInstance& VertexInstance = VertexInstances[ VertexInstanceID ];
-		verify( VertexInstance.ConnectedPolygons.Remove( PolygonID ) == 1 );
-
-		// Vertex instances which are only part of this contour (i.e. not shared by other polygons) are optionally deleted
-		if( bDeleteOrphanedVertexInstances && VertexInstance.ConnectedPolygons.Num() == 0 )
-		{
-			OrphanedVertexInstanceIDs.Add( VertexInstanceID );
-		}
-	}
-
-	// Remove reference to polygon from edges
-	static TArray<FEdgeID> EdgeIDs;
-	GetPolygonContourEdges( Contour, EdgeIDs );
-
-	FEdgeArray& Edges = GetMeshDescription()->Edges();
-
-	for( const FEdgeID EdgeID : EdgeIDs )
-	{
-		FMeshEdge& Edge = Edges[ EdgeID ];
-		verify( Edge.ConnectedPolygons.Remove( PolygonID ) == 1 );
-
-		// Identify orphaned edges
-		if( bDeleteOrphanedEdges && Edge.ConnectedPolygons.Num() == 0 )
-		{
-			OrphanedEdgeIDs.Add( EdgeID );
-		}
-	}
-}
-
-
 void UEditableMesh::DeletePolygons( const TArray<FPolygonID>& PolygonIDsToDelete, const bool bDeleteOrphanedEdges, const bool bDeleteOrphanedVertices, const bool bDeleteOrphanedVertexInstances, const bool bDeleteEmptyPolygonGroups )
 {
 	UE_LOG( LogEditableMesh, Verbose, TEXT( "DeletePolygons: %s" ), *LogHelpers::ArrayToString( PolygonIDsToDelete ) );
@@ -4575,27 +4470,10 @@ void UEditableMesh::DeletePolygons( const TArray<FPolygonID>& PolygonIDsToDelete
 
 		for( const FPolygonID PolygonID : PolygonIDsToDelete )
 		{
-			const FMeshPolygon& Polygon = Polygons[ PolygonID ];
-
-			DeletePolygonContour( PolygonID, Polygon.PerimeterContour, bDeleteOrphanedEdges, bDeleteOrphanedVertexInstances, OrphanedEdgeIDs, OrphanedVertexInstanceIDs );
-
-			for( const FMeshPolygonContour& HoleContour : Polygon.HoleContours )
-			{
-				DeletePolygonContour( PolygonID, HoleContour, bDeleteOrphanedEdges, bDeleteOrphanedVertexInstances, OrphanedEdgeIDs, OrphanedVertexInstanceIDs );
-			}
-
-			// Disconnect the polygon from the polygon group
-			const FPolygonGroupID PolygonGroupID = Polygon.PolygonGroupID;
-			FMeshPolygonGroup& PolygonGroup = PolygonGroups[ PolygonGroupID ];
-			verify( PolygonGroup.Polygons.Remove( PolygonID ) == 1 );
-
-			if( bDeleteEmptyPolygonGroups && PolygonGroup.Polygons.Num() == 0 )
-			{
-				EmptyPolygonGroupIDs.Add( PolygonGroupID );
-			}
-
-			// Delete the polygon
-			Polygons.Remove( PolygonID );
+			GetMeshDescription()->DeletePolygon( PolygonID, 
+				bDeleteOrphanedEdges ? &OrphanedEdgeIDs : nullptr,
+				bDeleteOrphanedVertexInstances ? &OrphanedVertexInstanceIDs : nullptr,
+				bDeleteEmptyPolygonGroups ? &EmptyPolygonGroupIDs : nullptr );
 
 			// Update spatial database
 			if( Octree.IsValid() )
@@ -4615,7 +4493,8 @@ void UEditableMesh::DeletePolygons( const TArray<FPolygonID>& PolygonIDsToDelete
 		// We do not want this to remove orphaned vertices; this will optionally happen below when removing edges.
 		if( OrphanedVertexInstanceIDs.Num() > 0 )
 		{
-			DeleteVertexInstances( OrphanedVertexInstanceIDs, false );
+			const bool bDeleteVertices = false;
+			DeleteVertexInstances( OrphanedVertexInstanceIDs, bDeleteVertices );
 		}
 
 		// Remove any edges which may have been orphaned. This may also optionally remove any orphaned vertices.
@@ -4649,37 +4528,30 @@ void UEditableMesh::CreatePolygonGroups( const TArray<FPolygonGroupToCreate>& Po
 {
 	UE_LOG( LogEditableMesh, Verbose, TEXT( "CreatePolygonGroups: %s" ), *LogHelpers::ArrayToString( PolygonGroupsToCreate ) );
 
-	FPolygonGroupArray& PolygonGroups = GetMeshDescription()->PolygonGroups();
-
-	// @todo mesheditor: Material, bEnableCollision and bCastShadow should all be implemented as PolygonGroup attributes
+	// Create polygon groups and initialize them
+	{
+		FPolygonGroupArray& PolygonGroups = GetMeshDescription()->PolygonGroups();
 
-	OutNewPolygonGroupIDs.Reset( PolygonGroupsToCreate.Num() );
+		OutNewPolygonGroupIDs.Reset( PolygonGroupsToCreate.Num() );
 
-	// Reserve elements
-	PolygonGroups.Reserve( PolygonGroups.Num() + PolygonGroupsToCreate.Num() );
+		// Reserve elements
+		GetMeshDescription()->ReserveNewPolygonGroups( PolygonGroupsToCreate.Num() );
 
-	for( const FPolygonGroupToCreate& PolygonGroupToCreate : PolygonGroupsToCreate )
-	{
-		// Allocate polygon group
-		FPolygonGroupID PolygonGroupID = PolygonGroupToCreate.OriginalPolygonGroupID;
-		if( PolygonGroupID != FPolygonGroupID::Invalid )
-		{
-			PolygonGroups.Insert( PolygonGroupID );
-		}
-		else
+		for( const FPolygonGroupToCreate& PolygonGroupToCreate : PolygonGroupsToCreate )
 		{
-			PolygonGroupID = PolygonGroups.Add();
-		}
-
-		FMeshPolygonGroup& PolygonGroup = PolygonGroups[ PolygonGroupID ];
-
-		PolygonGroup.MaterialAsset = PolygonGroupToCreate.MaterialAsset;
-		PolygonGroup.MaterialSlotName = PolygonGroupToCreate.MaterialSlotName;
-		PolygonGroup.ImportedMaterialSlotName = PolygonGroupToCreate.ImportedMaterialSlotName;
-		PolygonGroup.bEnableCollision = PolygonGroupToCreate.bEnableCollision;
-		PolygonGroup.bCastShadow = PolygonGroupToCreate.bCastShadow;
+			// Allocate polygon group
+			FPolygonGroupID PolygonGroupID = PolygonGroupToCreate.OriginalPolygonGroupID;
+			if( PolygonGroupID != FPolygonGroupID::Invalid )
+			{
+				GetMeshDescription()->CreatePolygonGroupWithID( PolygonGroupID );
+			}
+			else
+			{
+				PolygonGroupID = GetMeshDescription()->CreatePolygonGroup();
+			}
 
-		OutNewPolygonGroupIDs.Add( PolygonGroupID );
+			OutNewPolygonGroupIDs.Add( PolygonGroupID );
+		}
 	}
 
 	for( UEditableMeshAdapter* Adapter : Adapters )
@@ -4687,6 +4559,15 @@ void UEditableMesh::CreatePolygonGroups( const TArray<FPolygonGroupToCreate>& Po
 		Adapter->OnCreatePolygonGroups( this, OutNewPolygonGroupIDs );
 	}
 
+	// Apply attributes
+	for( int32 Index = 0; Index < OutNewPolygonGroupIDs.Num(); ++Index )
+	{
+		for( const FMeshElementAttributeData& EdgeAttribute : PolygonGroupsToCreate[ Index ].PolygonGroupAttributes.Attributes )
+		{
+			SetPolygonGroupAttribute( OutNewPolygonGroupIDs[ Index ], EdgeAttribute );
+		}
+	}
+
 	// Back up
 	{
 		FDeletePolygonGroupsChangeInput RevertInput;
@@ -4722,11 +4603,8 @@ void UEditableMesh::DeletePolygonGroups( const TArray<FPolygonGroupID>& PolygonG
 			FPolygonGroupToCreate& PolygonGroupToCreate = RevertInput.PolygonGroupsToCreate.Last();
 
 			PolygonGroupToCreate.OriginalPolygonGroupID = PolygonGroupID;
-			PolygonGroupToCreate.MaterialAsset = PolygonGroup.MaterialAsset.ToString();
-			PolygonGroupToCreate.MaterialSlotName = PolygonGroup.MaterialSlotName;
-			PolygonGroupToCreate.ImportedMaterialSlotName = PolygonGroup.ImportedMaterialSlotName;
-			PolygonGroupToCreate.bCastShadow = PolygonGroup.bCastShadow;
-			PolygonGroupToCreate.bEnableCollision = PolygonGroup.bEnableCollision;
+
+			BackupAllAttributes( PolygonGroupToCreate.PolygonGroupAttributes, GetMeshDescription()->PolygonGroupAttributes(), PolygonGroupID );
 		}
 
 		AddUndo( MakeUnique<FCreatePolygonGroupsChange>( MoveTemp( RevertInput ) ) );
@@ -4741,10 +4619,7 @@ void UEditableMesh::DeletePolygonGroups( const TArray<FPolygonGroupID>& PolygonG
 	{
 		for( const FPolygonGroupID PolygonGroupID : PolygonGroupIDs )
 		{
-			const FMeshPolygonGroup& PolygonGroup = PolygonGroups[ PolygonGroupID ];
-			check( PolygonGroup.Polygons.Num() == 0 );
-
-			PolygonGroups.Remove( PolygonGroupID );
+			GetMeshDescription()->DeletePolygonGroup( PolygonGroupID );
 		}
 	}
 }
@@ -4759,24 +4634,19 @@ void UEditableMesh::SetVerticesAttributes( const TArray<FAttributesForVertex>& A
 	RevertInput.AttributesForVertices.Reserve( AttributesForVertices.Num() );
 	for( const FAttributesForVertex& AttributesForVertex : AttributesForVertices )
 	{
-		FAttributesForVertex& RevertVertex = *new( RevertInput.AttributesForVertices ) FAttributesForVertex();
-		RevertVertex.VertexID = AttributesForVertex.VertexID;
+		const FVertexID VertexID = AttributesForVertex.VertexID;
 
-		// Set the vertex attributes
-		{
-			RevertVertex.VertexAttributes.Attributes.Reserve( AttributesForVertex.VertexAttributes.Attributes.Num() );
+		RevertInput.AttributesForVertices.Emplace();
+		FAttributesForVertex& RevertVertex = RevertInput.AttributesForVertices.Last();
+		RevertVertex.VertexID = VertexID;
 
-			for( const FMeshElementAttributeData& VertexAttribute : AttributesForVertex.VertexAttributes.Attributes )
-			{
-				// Back up the attribute
-				FMeshElementAttributeData& RevertAttribute = *new( RevertVertex.VertexAttributes.Attributes ) FMeshElementAttributeData(
-					VertexAttribute.AttributeName,
-					VertexAttribute.AttributeIndex,
-					GetVertexAttribute( AttributesForVertex.VertexID, VertexAttribute.AttributeName, VertexAttribute.AttributeIndex ) );
+		// Back up the attributes
+		BackupAttributesInList( RevertVertex.VertexAttributes, AttributesForVertex.VertexAttributes, GetMeshDescription()->VertexAttributes(), VertexID );
 
-				// Set the new attribute
-				SetVertexAttribute( AttributesForVertex.VertexID, VertexAttribute.AttributeName, VertexAttribute.AttributeIndex, VertexAttribute.AttributeValue );
-			}
+		for( const FMeshElementAttributeData& VertexAttribute : AttributesForVertex.VertexAttributes.Attributes )
+		{
+			// Set the new attribute
+			SetVertexAttribute( VertexID, VertexAttribute );
 		}
 	}
 
@@ -4795,24 +4665,19 @@ void UEditableMesh::SetVertexInstancesAttributes( const TArray<FAttributesForVer
 	RevertInput.AttributesForVertexInstances.Reserve( AttributesForVertexInstances.Num() );
 	for( const FAttributesForVertexInstance& AttributesForVertexInstance : AttributesForVertexInstances )
 	{
-		FAttributesForVertexInstance& RevertVertexInstance = *new( RevertInput.AttributesForVertexInstances ) FAttributesForVertexInstance();
-		RevertVertexInstance.VertexInstanceID = AttributesForVertexInstance.VertexInstanceID;
+		const FVertexInstanceID VertexInstanceID = AttributesForVertexInstance.VertexInstanceID;
 
-		// Set the vertex instance attributes
-		{
-			RevertVertexInstance.VertexInstanceAttributes.Attributes.Reserve( AttributesForVertexInstance.VertexInstanceAttributes.Attributes.Num() );
+		RevertInput.AttributesForVertexInstances.Emplace();
+		FAttributesForVertexInstance& RevertVertexInstance = RevertInput.AttributesForVertexInstances.Last();
+		RevertVertexInstance.VertexInstanceID = VertexInstanceID;
 
-			for( const FMeshElementAttributeData& VertexInstanceAttribute : AttributesForVertexInstance.VertexInstanceAttributes.Attributes )
-			{
-				// Back up the attribute
-				FMeshElementAttributeData& RevertAttribute = *new( RevertVertexInstance.VertexInstanceAttributes.Attributes ) FMeshElementAttributeData(
-					VertexInstanceAttribute.AttributeName,
-					VertexInstanceAttribute.AttributeIndex,
-					GetVertexInstanceAttribute( AttributesForVertexInstance.VertexInstanceID, VertexInstanceAttribute.AttributeName, VertexInstanceAttribute.AttributeIndex ) );
+		// Back up the attributes
+		BackupAttributesInList( RevertVertexInstance.VertexInstanceAttributes, AttributesForVertexInstance.VertexInstanceAttributes, GetMeshDescription()->VertexInstanceAttributes(), VertexInstanceID );
 
-				// Set the new attribute
-				SetVertexInstanceAttribute( AttributesForVertexInstance.VertexInstanceID, VertexInstanceAttribute.AttributeName, VertexInstanceAttribute.AttributeIndex, VertexInstanceAttribute.AttributeValue );
-			}
+		for( const FMeshElementAttributeData& VertexInstanceAttribute : AttributesForVertexInstance.VertexInstanceAttributes.Attributes )
+		{
+			// Set the new attribute
+			SetVertexInstanceAttribute( VertexInstanceID, VertexInstanceAttribute );
 		}
 	}
 
@@ -4831,25 +4696,21 @@ void UEditableMesh::SetEdgesAttributes( const TArray<FAttributesForEdge>& Attrib
 	RevertInput.AttributesForEdges.Reserve( AttributesForEdges.Num() );
 	for( const FAttributesForEdge& AttributesForEdge : AttributesForEdges )
 	{
-		FAttributesForEdge& RevertEdge = *new( RevertInput.AttributesForEdges ) FAttributesForEdge();
-		RevertEdge.EdgeID = AttributesForEdge.EdgeID;
-		RevertEdge.bIsUndo = true;
+		const FEdgeID EdgeID = AttributesForEdge.EdgeID;
 
-		// Set the edge attributes
-		{
-			RevertEdge.EdgeAttributes.Attributes.Reserve( AttributesForEdge.EdgeAttributes.Attributes.Num() );
+		RevertInput.AttributesForEdges.Emplace();
+		FAttributesForEdge& RevertEdge = RevertInput.AttributesForEdges.Last();
+		RevertEdge.EdgeID = EdgeID;
+		RevertEdge.bIsUndo = true;
 
-			for( const FMeshElementAttributeData& EdgeAttribute : AttributesForEdge.EdgeAttributes.Attributes )
-			{
-				// Back up the attribute
-				FMeshElementAttributeData& RevertAttribute = *new( RevertEdge.EdgeAttributes.Attributes ) FMeshElementAttributeData(
-					EdgeAttribute.AttributeName,
-					EdgeAttribute.AttributeIndex,
-					GetEdgeAttribute( AttributesForEdge.EdgeID, EdgeAttribute.AttributeName, EdgeAttribute.AttributeIndex ) );
+		// Back up the attributes
+		BackupAttributesInList( RevertEdge.EdgeAttributes, AttributesForEdge.EdgeAttributes, GetMeshDescription()->EdgeAttributes(), EdgeID );
 
-				// Set the new attribute
-				SetEdgeAttribute( AttributesForEdge.EdgeID, EdgeAttribute.AttributeName, EdgeAttribute.AttributeIndex, EdgeAttribute.AttributeValue, AttributesForEdge.bIsUndo );
-			}
+		for( const FMeshElementAttributeData& EdgeAttribute : AttributesForEdge.EdgeAttributes.Attributes )
+		{
+			// Set the new attribute
+			// @todo mesheditor: having to pass bIsUndo here is horrible. Find a better way.
+			SetEdgeAttribute( AttributesForEdge.EdgeID, EdgeAttribute, AttributesForEdge.bIsUndo );
 		}
 	}
 
@@ -4859,49 +4720,6 @@ void UEditableMesh::SetEdgesAttributes( const TArray<FAttributesForEdge>& Attrib
 }
 
 
-const TArray<FName>& UEditableMesh::GetMergeableVertexInstanceAttributes()
-{
-	static TArray<FName> MergeableVertexInstanceAttributes;
-	if( MergeableVertexInstanceAttributes.Num() == 0 )
-	{
-		MergeableVertexInstanceAttributes.Add( UEditableMeshAttribute::VertexTextureCoordinate() );
-		MergeableVertexInstanceAttributes.Add( UEditableMeshAttribute::VertexColor() );
-	}
-	return MergeableVertexInstanceAttributes;
-}
-
-
-void UEditableMesh::GetMergedNamedVertexInstanceAttributeData( const TArray<FName>& AttributeNames, const FVertexInstanceID VertexInstanceID, const TArray<FMeshElementAttributeData>& AttributeValuesToMerge, TArray<FMeshElementAttributeData>& OutAttributes ) const
-{
-	OutAttributes.Reset();
-
-	for( const FName AttributeName : AttributeNames )
-	{
-		const int32 MaxAttributeIndex = GetMaxAttributeIndex( AttributeName );
-		for( int32 AttributeIndex = 0; AttributeIndex < MaxAttributeIndex; ++AttributeIndex )
-		{
-			const int32 ReplaceIndex =  AttributeValuesToMerge.IndexOfByPredicate(
-				[ AttributeName, AttributeIndex ]( const FMeshElementAttributeData& AttributeData )
-				{
-					return AttributeData.AttributeName == AttributeName && AttributeData.AttributeIndex == AttributeIndex;
-				}
-			);
-
-			if( ReplaceIndex == INDEX_NONE )
-			{
-				// If the named attribute and index doesn't appear in the list to merge, add the current attribute value to the list
-				OutAttributes.Emplace( AttributeName, AttributeIndex, GetVertexInstanceAttribute( VertexInstanceID, AttributeName, AttributeIndex ) );
-			}
-			else
-			{
-				// If the named attribute and index does appear in the list to merge, add that to the list instead
-				OutAttributes.Emplace( AttributeValuesToMerge[ ReplaceIndex ] );
-			}
-		}
-	}
-}
-
-
 void UEditableMesh::ChangePolygonsVertexInstances( const TArray<FChangeVertexInstancesForPolygon>& VertexInstancesForPolygons )
 {
 	UE_LOG( LogEditableMesh, Verbose, TEXT( "ChangePolygonsVertexInstances: %s" ), *LogHelpers::ArrayToString( VertexInstancesForPolygons ) );
@@ -5056,9 +4874,10 @@ void UEditableMesh::GetConnectedSoftEdges( const FVertexID VertexID, TArray<FEdg
 {
 	OutConnectedSoftEdges.Reset();
 
+	const TEdgeAttributeArray<bool>& EdgeHardnesses = GetMeshDescription()->EdgeAttributes().GetAttributes<bool>( MeshAttribute::Edge::IsHard );
 	for( const FEdgeID ConnectedEdgeID : GetMeshDescription()->GetVertex( VertexID ).ConnectedEdgeIDs )
 	{
-		if( GetEdgeAttribute( ConnectedEdgeID, UEditableMeshAttribute::EdgeIsHard(), 0 ).X == 0.0f )
+		if( !EdgeHardnesses[ ConnectedEdgeID ] )
 		{
 			OutConnectedSoftEdges.Add( ConnectedEdgeID );
 		}
@@ -5149,31 +4968,12 @@ void UEditableMesh::SetPolygonsVertexAttributes( const TArray<FVertexAttributesF
 }
 
 
-bool UEditableMesh::AreAttributeListsEqual( const TArray<FMeshElementAttributeData>& A, const TArray<FMeshElementAttributeData>& B ) const
-{
-	// Check that two attribute lists have attributes of equal values.
-	// It is expected that each list contains exactly the same attribute names and indices, in the same order.
-	check( A.Num() == B.Num() );
-	for( int32 Index = 0; Index < A.Num(); ++Index )
-	{
-		check( A[ Index ].AttributeName == B[ Index ].AttributeName );
-		check( A[ Index ].AttributeIndex == B[ Index ].AttributeIndex );
-		if( A[ Index ].AttributeValue != B[ Index ].AttributeValue )
-		{
-			return false;
-		}
-	}
-	return true;
-};
-
-
 void UEditableMesh::SetPolygonContourVertexAttributes( FMeshPolygonContour& Contour, const FPolygonID PolygonID, const int32 HoleIndex, const TArray<FMeshElementAttributeList>& AttributeLists )
 {
 	const int32 NumContourVertices = Contour.VertexInstanceIDs.Num();
 	check( AttributeLists.Num() == NumContourVertices );
 
 	// Iterate round all polygons in the contour
-	int32 PrevIndex = NumContourVertices - 1;
 	for( int32 Index = 0; Index < NumContourVertices; ++Index )
 	{
 		const FMeshElementAttributeList& AttributeList = AttributeLists[ Index ];
@@ -5188,46 +4988,36 @@ void UEditableMesh::SetPolygonContourVertexAttributes( FMeshPolygonContour& Cont
 		const FVertexInstanceID VertexInstanceID = Contour.VertexInstanceIDs[ Index ];
 		const FVertexID VertexID = GetVertexInstanceVertex( VertexInstanceID );
 
-		// We need to determine whether changing the attributes on this polygon vertex would require vertex instances to be split or merged.
-		// First determine all the existing vertex instances which are in this soft edged group, that is to say,
-		// the group of polygons delimited by hard edges, of which this one is a member.
-		static TArray<FVertexInstanceID> PossibleVertexInstances;
-		PossibleVertexInstances.Reset();
-		const bool bPolygonNotYetInitialized = false;
-		GetVertexInstancesInSameSoftEdgedGroup( VertexID, PolygonID, bPolygonNotYetInitialized, PossibleVertexInstances );
-
-		// Now we have various possibilities:
-		//   1) If any of these vertex instances' attributes match our new polygon vertex attributes, just merge with that one, and disconnect the old one from the polygon.
-		//     1a) If that would cause the old vertex instance to become orphaned (because it is no longer connected to any polygon), it should be deleted.
-		// otherwise:
-		//   2) There are no vertex instances to merge into. 
-		//     2a) If the current vertex instance is not used by any other polygon, change its attributes in-place.
-		//   otherwise:
-		//     2b) Create a new vertex instance with the new attributes, and disconnect the old one from the polygon.
-
-		// Build the complete list of attributes for the polygon vertex being changed
-		static TArray<FMeshElementAttributeData> AllAttributes;
-		GetMergedNamedVertexInstanceAttributeData( GetMergeableVertexInstanceAttributes(), VertexInstanceID, AttributeList.Attributes, AllAttributes );
+		const TArray<FPolygonID>& ConnectedPolygons = GetMeshDescription()->GetVertexInstanceConnectedPolygons( VertexInstanceID );
+		check( ConnectedPolygons.Contains( PolygonID ) );
 
-		// First look for a vertex instance with matching attributes to ours
-		FVertexInstanceID MergedVertexInstanceID = FVertexInstanceID::Invalid;
-
-		for( const FVertexInstanceID PossibleVertexInstance : PossibleVertexInstances )
+		if( ConnectedPolygons.Num() == 1 )
 		{
-			static TArray<FMeshElementAttributeData> AllAttributesToCompare;
-			GetMergedNamedVertexInstanceAttributeData( GetMergeableVertexInstanceAttributes(), PossibleVertexInstance, TArray<FMeshElementAttributeData>(), AllAttributesToCompare );
+			// This is the only polygon using this vertex instance, so change it in place
+			static TArray<FAttributesForVertexInstance> AttributesForVertexInstance;
+			AttributesForVertexInstance.Reset();
+			AttributesForVertexInstance.SetNum( 1 );
+			AttributesForVertexInstance[ 0 ].VertexInstanceID = VertexInstanceID;
+			AttributesForVertexInstance[ 0 ].VertexInstanceAttributes = AttributeList;
 
-			if( AreAttributeListsEqual( AllAttributes, AllAttributesToCompare ) )
-			{
-				// Found a vertex instance we can merge into
-				MergedVertexInstanceID = PossibleVertexInstance;
-				break;
-			}
+			SetVertexInstancesAttributes( AttributesForVertexInstance );
 		}
-
-		if( MergedVertexInstanceID != FVertexInstanceID::Invalid )
+		else
 		{
-			// If we found a vertex instance we can merge into, change the perimeter vertex instance for the new one
+			// Split vertex instance: create a new one
+			static TArray<FVertexInstanceToCreate> VertexInstancesToCreate;
+			VertexInstancesToCreate.Reset();
+			VertexInstancesToCreate.SetNum( 1 );
+			VertexInstancesToCreate[ 0 ].VertexID = VertexID;
+			BackupAllAttributes( VertexInstancesToCreate[ 0 ].VertexInstanceAttributes, GetMeshDescription()->VertexInstanceAttributes(), VertexInstanceID );
+
+			// Override original attributes with new ones by putting them at the end of the list
+			VertexInstancesToCreate[ 0 ].VertexInstanceAttributes.Attributes.Append( AttributeList.Attributes );
+
+			static TArray<FVertexInstanceID> NewVertexInstanceIDs;
+			CreateVertexInstances( VertexInstancesToCreate, NewVertexInstanceIDs );
+
+			// and set it on the contour
 			static TArray<FChangeVertexInstancesForPolygon> VertexInstancesToChange;
 			VertexInstancesToChange.Reset();
 			VertexInstancesToChange.SetNum( 1 );
@@ -5237,80 +5027,20 @@ void UEditableMesh::SetPolygonContourVertexAttributes( FMeshPolygonContour& Cont
 			{
 				VertexInstancesToChange[ 0 ].PerimeterVertexIndicesAndInstanceIDs.Emplace();
 				VertexInstancesToChange[ 0 ].PerimeterVertexIndicesAndInstanceIDs[ 0 ].ContourIndex = Index;
-				VertexInstancesToChange[ 0 ].PerimeterVertexIndicesAndInstanceIDs[ 0 ].VertexInstanceID = MergedVertexInstanceID;
+				VertexInstancesToChange[ 0 ].PerimeterVertexIndicesAndInstanceIDs[ 0 ].VertexInstanceID = NewVertexInstanceIDs[ 0 ];
 			}
 			else
 			{
 				VertexInstancesToChange[ 0 ].VertexIndicesAndInstanceIDsForEachHole.SetNum( HoleIndex + 1 );
 				VertexInstancesToChange[ 0 ].VertexIndicesAndInstanceIDsForEachHole[ HoleIndex ].VertexIndicesAndInstanceIDs.Emplace();
 				VertexInstancesToChange[ 0 ].VertexIndicesAndInstanceIDsForEachHole[ HoleIndex ].VertexIndicesAndInstanceIDs[ 0 ].ContourIndex = Index;
-				VertexInstancesToChange[ 0 ].VertexIndicesAndInstanceIDsForEachHole[ HoleIndex ].VertexIndicesAndInstanceIDs[ 0 ].VertexInstanceID = MergedVertexInstanceID;
+				VertexInstancesToChange[ 0 ].VertexIndicesAndInstanceIDsForEachHole[ HoleIndex ].VertexIndicesAndInstanceIDs[ 0 ].VertexInstanceID = NewVertexInstanceIDs[ 0 ];
 			}
 
 			ChangePolygonsVertexInstances( VertexInstancesToChange );
 
-			// If the change orphaned the old one, delete it
-			// @todo mesheditor: make this optional?
-			if( GetVertexInstanceConnectedPolygonCount( VertexInstanceID ) == 0 )
-			{
-				static TArray<FVertexInstanceID> VertexInstanceIDsToDelete;
-				VertexInstanceIDsToDelete.SetNumUninitialized( 1 );
-				VertexInstanceIDsToDelete[ 0 ] = VertexInstanceID;
-
-				const bool bDeleteOrphanedVertices = false;
-				DeleteVertexInstances( VertexInstanceIDsToDelete, bDeleteOrphanedVertices );
-			}
+			// @todo mesheditor: weld identical vertex instances in the same smoothing group
 		}
-		else
-		{
-			if( GetVertexInstanceConnectedPolygonCount( VertexInstanceID ) == 1 )
-			{
-				// This is the only polygon using this vertex instance, so change it in place
-				static TArray<FAttributesForVertexInstance> AttributesForVertexInstance;
-				AttributesForVertexInstance.Reset();
-				AttributesForVertexInstance.SetNum( 1 );
-				AttributesForVertexInstance[ 0 ].VertexInstanceID = VertexInstanceID;
-				AttributesForVertexInstance[ 0 ].VertexInstanceAttributes = AttributeList;
-
-				SetVertexInstancesAttributes( AttributesForVertexInstance );
-			}
-			else
-			{
-				// Split vertex instance: create a new one
-				static TArray<FVertexInstanceToCreate> VertexInstancesToCreate;
-				VertexInstancesToCreate.Reset();
-				VertexInstancesToCreate.SetNum( 1 );
-				VertexInstancesToCreate[ 0 ].VertexID = VertexID;
-				VertexInstancesToCreate[ 0 ].VertexInstanceAttributes.Attributes = AllAttributes;
-
-				static TArray<FVertexInstanceID> NewVertexInstanceIDs;
-				CreateVertexInstances( VertexInstancesToCreate, NewVertexInstanceIDs );
-
-				// and set it on the contour
-				static TArray<FChangeVertexInstancesForPolygon> VertexInstancesToChange;
-				VertexInstancesToChange.Reset();
-				VertexInstancesToChange.SetNum( 1 );
-
-				VertexInstancesToChange[ 0 ].PolygonID = PolygonID;
-				if( HoleIndex == INDEX_NONE )
-				{
-					VertexInstancesToChange[ 0 ].PerimeterVertexIndicesAndInstanceIDs.Emplace();
-					VertexInstancesToChange[ 0 ].PerimeterVertexIndicesAndInstanceIDs[ 0 ].ContourIndex = Index;
-					VertexInstancesToChange[ 0 ].PerimeterVertexIndicesAndInstanceIDs[ 0 ].VertexInstanceID = NewVertexInstanceIDs[ 0 ];
-				}
-				else
-				{
-					VertexInstancesToChange[ 0 ].VertexIndicesAndInstanceIDsForEachHole.SetNum( HoleIndex + 1 );
-					VertexInstancesToChange[ 0 ].VertexIndicesAndInstanceIDsForEachHole[ HoleIndex ].VertexIndicesAndInstanceIDs.Emplace();
-					VertexInstancesToChange[ 0 ].VertexIndicesAndInstanceIDsForEachHole[ HoleIndex ].VertexIndicesAndInstanceIDs[ 0 ].ContourIndex = Index;
-					VertexInstancesToChange[ 0 ].VertexIndicesAndInstanceIDsForEachHole[ HoleIndex ].VertexIndicesAndInstanceIDs[ 0 ].VertexInstanceID = NewVertexInstanceIDs[ 0 ];
-				}
-
-				ChangePolygonsVertexInstances( VertexInstancesToChange );
-			}
-		}
-
-		PrevIndex = Index;
 	}
 }
 
@@ -5357,13 +5087,11 @@ void UEditableMesh::TryToRemovePolygonEdge( const FEdgeID EdgeID, bool& bOutWasE
 
 		if( bBothVerticesConnectToPolygonEdges )
 		{
-			TArray<FVertexID> PolygonAVertexIDs, PolygonBVertexIDs;
-
 			const FPolygonID PolygonAID = GetEdgeConnectedPolygon( EdgeID, 0 );
-			GetPolygonPerimeterVertices( PolygonAID, /* Out */ PolygonAVertexIDs );
-
 			const FPolygonID PolygonBID = GetEdgeConnectedPolygon( EdgeID, 1 );
-			GetPolygonPerimeterVertices( PolygonBID, /* Out */ PolygonBVertexIDs );
+
+			const TArray<FVertexInstanceID>& PolygonAVertexInstanceIDs = GetMeshDescription()->GetPolygonPerimeterVertexInstances( PolygonAID );
+			const TArray<FVertexInstanceID>& PolygonBVertexInstanceIDs = GetMeshDescription()->GetPolygonPerimeterVertexInstances( PolygonBID );
 
 			// If the polygons are in different polygon groups, we can't remove the edge because we can't determine
 			// which polygon group the replacing polygon should belong to
@@ -5381,12 +5109,17 @@ void UEditableMesh::TryToRemovePolygonEdge( const FEdgeID EdgeID, bool& bOutWasE
 					// end up being two less than the total vertex count of the two polygons we're replacing.  However, this
 					// might not be true for cases where either polygon shares more than one edge with the other, right?
 
+					// @todo mesheditor: Case study: Two tris joined by an edge with an extra vertex inserted, then try to remove
+					// one of those diagonal edges.
+
 					// Find the edge vertices in the first polygon
 					int32 EdgeStartsAtVertexInPolygonA = INDEX_NONE;
-					for( int32 PolygonAVertexNumber = 0; PolygonAVertexNumber < PolygonAVertexIDs.Num(); ++PolygonAVertexNumber )
+					for( int32 PolygonAVertexNumber = 0; PolygonAVertexNumber < PolygonAVertexInstanceIDs.Num(); ++PolygonAVertexNumber )
 					{
-						const FVertexID PolygonAVertexID = PolygonAVertexIDs[ PolygonAVertexNumber ];
-						const FVertexID PolygonANextVertexID = PolygonAVertexIDs[ ( PolygonAVertexNumber + 1 ) % PolygonAVertexIDs.Num() ];
+						const FVertexInstanceID PolygonAVertexInstanceID = PolygonAVertexInstanceIDs[ PolygonAVertexNumber ];
+						const FVertexInstanceID PolygonANextVertexInstanceID = PolygonAVertexInstanceIDs[ ( PolygonAVertexNumber + 1 ) % PolygonAVertexInstanceIDs.Num() ];
+						const FVertexID PolygonAVertexID = GetMeshDescription()->GetVertexInstanceVertex( PolygonAVertexInstanceID );
+						const FVertexID PolygonANextVertexID = GetMeshDescription()->GetVertexInstanceVertex( PolygonANextVertexInstanceID );
 
 						if( ( PolygonAVertexID == EdgeVertexIDA || PolygonAVertexID == EdgeVertexIDB ) &&
 							( PolygonANextVertexID == EdgeVertexIDA || PolygonANextVertexID == EdgeVertexIDB ) )
@@ -5396,15 +5129,17 @@ void UEditableMesh::TryToRemovePolygonEdge( const FEdgeID EdgeID, bool& bOutWasE
 						}
 					}
 					check( EdgeStartsAtVertexInPolygonA != INDEX_NONE );
-					const int32 EdgeEndsAtVertexInPolygonA = ( EdgeStartsAtVertexInPolygonA + 1 ) % PolygonAVertexIDs.Num();
+					const int32 EdgeEndsAtVertexInPolygonA = ( EdgeStartsAtVertexInPolygonA + 1 ) % PolygonAVertexInstanceIDs.Num();
 
 
 					// Find the edge vertices in the second polygon
 					int32 EdgeStartsAtVertexInPolygonB = INDEX_NONE;
-					for( int32 PolygonBVertexNumber = 0; PolygonBVertexNumber < PolygonBVertexIDs.Num(); ++PolygonBVertexNumber )
+					for( int32 PolygonBVertexNumber = 0; PolygonBVertexNumber < PolygonBVertexInstanceIDs.Num(); ++PolygonBVertexNumber )
 					{
-						const FVertexID PolygonBVertexID = PolygonBVertexIDs[ PolygonBVertexNumber ];
-						const FVertexID PolygonBNextVertexID = PolygonBVertexIDs[ ( PolygonBVertexNumber + 1 ) % PolygonBVertexIDs.Num() ];
+						const FVertexInstanceID PolygonBVertexInstanceID = PolygonBVertexInstanceIDs[ PolygonBVertexNumber ];
+						const FVertexInstanceID PolygonBNextVertexInstanceID = PolygonBVertexInstanceIDs[ ( PolygonBVertexNumber + 1 ) % PolygonBVertexInstanceIDs.Num() ];
+						const FVertexID PolygonBVertexID = GetMeshDescription()->GetVertexInstanceVertex( PolygonBVertexInstanceID );
+						const FVertexID PolygonBNextVertexID = GetMeshDescription()->GetVertexInstanceVertex( PolygonBNextVertexInstanceID );
 
 						if( ( PolygonBVertexID == EdgeVertexIDA || PolygonBVertexID == EdgeVertexIDB ) &&
 							( PolygonBNextVertexID == EdgeVertexIDA || PolygonBNextVertexID == EdgeVertexIDB ) )
@@ -5414,93 +5149,44 @@ void UEditableMesh::TryToRemovePolygonEdge( const FEdgeID EdgeID, bool& bOutWasE
 						}
 					}
 					check( EdgeStartsAtVertexInPolygonB != INDEX_NONE );
-					const int32 EdgeEndsAtVertexInPolygonB = ( EdgeStartsAtVertexInPolygonB + 1 ) % PolygonBVertexIDs.Num();
+					const int32 EdgeEndsAtVertexInPolygonB = ( EdgeStartsAtVertexInPolygonB + 1 ) % PolygonBVertexInstanceIDs.Num();
 
 					
 					// Do the polygons wind in the same direction?  If they do, the edge order will be reversed.
-					const bool bPolygonsWindInSameDirection = PolygonAVertexIDs[ EdgeStartsAtVertexInPolygonA ] != PolygonBVertexIDs[ EdgeStartsAtVertexInPolygonB ];
+					const FVertexID PolygonAStartVertex = GetMeshDescription()->GetVertexInstanceVertex( PolygonAVertexInstanceIDs[ EdgeStartsAtVertexInPolygonA ] );
+					const FVertexID PolygonBStartVertex = GetMeshDescription()->GetVertexInstanceVertex( PolygonBVertexInstanceIDs[ EdgeStartsAtVertexInPolygonB ] );
+					const bool bPolygonsWindInSameDirection = PolygonAStartVertex != PolygonBStartVertex;
 
 
 					// Start adding vertices from the first polygon, starting with the vertex right after the edge we're removing.
 					// We'll continue to add vertices from this polygon until we reach back around to that edge.
-					const int32 PolygonAStartVertex = EdgeEndsAtVertexInPolygonA;
-					const int32 PolygonAEndVertex = EdgeStartsAtVertexInPolygonA;
-					for( int32 PolygonAVertexNumber = PolygonAStartVertex; PolygonAVertexNumber != PolygonAEndVertex; )
-					{
-						const FVertexID PolygonAVertexID = PolygonAVertexIDs[ PolygonAVertexNumber ];
-
-						FVertexAndAttributes& NewPolygonVertex = *new( NewPolygonVertices ) FVertexAndAttributes();
-						NewPolygonVertex.VertexID = PolygonAVertexID;
-
-						// Store off vertex attributes so we can re-apply them later to the new polygon
-						// @todo mesheditor: utility method which automatically does this with all known polygon vertex attributes? (or redundant when using "uber mesh"?)
-						for( int32 TextureCoordinateIndex = 0; TextureCoordinateIndex < TextureCoordinateCount; ++TextureCoordinateIndex )
-						{
-							NewPolygonVertex.PolygonVertexAttributes.Attributes.Emplace(
-								UEditableMeshAttribute::VertexTextureCoordinate(),
-								TextureCoordinateIndex,
-								GetPolygonPerimeterVertexAttribute( PolygonAID, PolygonAVertexNumber, UEditableMeshAttribute::VertexTextureCoordinate(), TextureCoordinateIndex )
-							);
-						}
-
-						NewPolygonVertex.PolygonVertexAttributes.Attributes.Emplace(
-							UEditableMeshAttribute::VertexColor(),
-							0,
-							GetPolygonPerimeterVertexAttribute( PolygonAID, PolygonAVertexNumber, UEditableMeshAttribute::VertexColor(), 0 )
-						);
-
-						PolygonAVertexNumber = ( PolygonAVertexNumber + 1 ) % PolygonAVertexIDs.Num();
-					}
-
-					// Now add vertices from the second polygon
-					const int32 PolygonBStartVertex = bPolygonsWindInSameDirection ? EdgeEndsAtVertexInPolygonB : EdgeStartsAtVertexInPolygonB;
-					const int32 PolygonBEndVertex = bPolygonsWindInSameDirection ? EdgeStartsAtVertexInPolygonB : EdgeEndsAtVertexInPolygonB;
-					for( int32 PolygonBVertexNumber = PolygonBStartVertex; PolygonBVertexNumber != PolygonBEndVertex; )
+					const int32 PolygonAStartVertexNumber = EdgeEndsAtVertexInPolygonA;
+					const int32 PolygonAEndVertexNumber = EdgeStartsAtVertexInPolygonA;
+					for( int32 PolygonAVertexNumber = PolygonAStartVertexNumber; PolygonAVertexNumber != PolygonAEndVertexNumber; )
 					{
-						const FVertexID PolygonBVertexID = PolygonBVertexIDs[ PolygonBVertexNumber ];
+						NewPolygonVertices.Emplace();
+						FVertexAndAttributes& NewPolygonVertex = NewPolygonVertices.Last();
+						NewPolygonVertex.VertexInstanceID = PolygonAVertexInstanceIDs[ PolygonAVertexNumber ];
 
-						FVertexAndAttributes& NewPolygonVertex = *new( NewPolygonVertices ) FVertexAndAttributes();
-						NewPolygonVertex.VertexID = PolygonBVertexID;
-
-						// Store off vertex attributes so we can re-apply them later to the new polygon
-						for( int32 TextureCoordinateIndex = 0; TextureCoordinateIndex < TextureCoordinateCount; ++TextureCoordinateIndex )
-						{
-							NewPolygonVertex.PolygonVertexAttributes.Attributes.Emplace(
-								UEditableMeshAttribute::VertexTextureCoordinate(),
-								TextureCoordinateIndex,
-								GetPolygonPerimeterVertexAttribute( PolygonBID, PolygonBVertexNumber, UEditableMeshAttribute::VertexTextureCoordinate(), TextureCoordinateIndex )
-							);
-						}
+						PolygonAVertexNumber = ( PolygonAVertexNumber + 1 ) % PolygonAVertexInstanceIDs.Num();
+					}
 
-						NewPolygonVertex.PolygonVertexAttributes.Attributes.Emplace(
-							UEditableMeshAttribute::VertexColor(),
-							0,
-							GetPolygonPerimeterVertexAttribute( PolygonBID, PolygonBVertexNumber, UEditableMeshAttribute::VertexColor(), 0 )
-						);
+					// Now add vertices from the second polygon
+					const int32 PolygonBStartVertexNumber = bPolygonsWindInSameDirection ? EdgeEndsAtVertexInPolygonB : EdgeStartsAtVertexInPolygonB;
+					const int32 PolygonBEndVertexNumber = bPolygonsWindInSameDirection ? EdgeStartsAtVertexInPolygonB : EdgeEndsAtVertexInPolygonB;
+					const int32 PolygonBNumberIncrement = bPolygonsWindInSameDirection ? 1 : PolygonBVertexInstanceIDs.Num() - 1;
+					for( int32 PolygonBVertexNumber = PolygonBStartVertexNumber; PolygonBVertexNumber != PolygonBEndVertexNumber; )
+					{
+						NewPolygonVertices.Emplace();
+						FVertexAndAttributes& NewPolygonVertex = NewPolygonVertices.Last();
+						NewPolygonVertex.VertexInstanceID = PolygonBVertexInstanceIDs[ PolygonBVertexNumber ];
 
-						if( bPolygonsWindInSameDirection )
-						{
-							// Iterate forwards
-							PolygonBVertexNumber = ( PolygonBVertexNumber + 1 ) % PolygonBVertexIDs.Num();
-						}
-						else
-						{
-							// Iterate backwards
-							if( PolygonBVertexNumber == 0 )	   // @todo mesheditor: Not tested yet.  Need to connect two faces that wind in opposite directions
-							{
-								PolygonBVertexNumber = PolygonBVertexIDs.Num() - 1;
-							}
-							else
-							{
-								--PolygonBVertexNumber;
-							}
-						}
+						PolygonBVertexNumber = ( PolygonBVertexNumber + PolygonBNumberIncrement ) % PolygonBVertexInstanceIDs.Num();
 					}
 				}
 
 				// OK, we can go ahead and delete the edge and its connected polygons.  We do NOT want to delete any orphaned
 				// edges or vertices though.  We're going to create a new polygon that connects to those right afterwards.
-				// @todo mesheditor: simplify this action so we just create the new polygon from existing vertex instances, instead of remembering all the attributes.
 				const bool bDeleteOrphanedEdges = false;
 				const bool bDeleteOrphanedVertices = false;
 				const bool bDeleteOrphanedVertexInstances = false;
@@ -5561,22 +5247,12 @@ void UEditableMesh::TryToRemoveVertex( const FVertexID VertexID, bool& bOutWasVe
 
 		// Try to preserve attributes of the edges we're deleting.  We'll take the attributes from the first
 		// edge and apply them to the newly created edge
-		static TArray<FMeshElementAttributeData> NewEdgeAttributes;
+		static FMeshElementAttributeList EdgeAttributeList;
 		{
 			const FEdgeID OtherEdgeID = GetVertexConnectedEdge( VertexID, 0 );		// @todo mesheditor: This is sort of subjective how we handle this (taking attributes from the first edge)
 
-			NewEdgeAttributes.Reset();
-			for( const FName AttributeName : UEditableMesh::GetValidEdgeAttributes() )
-			{
-				const int32 MaxAttributeIndex = GetMaxAttributeIndex( AttributeName );
-				for( int32 AttributeIndex = 0; AttributeIndex < MaxAttributeIndex; ++AttributeIndex )
-				{
-					FMeshElementAttributeData& EdgeAttribute = *new( NewEdgeAttributes ) FMeshElementAttributeData(
-						AttributeName,
-						AttributeIndex,
-						GetEdgeAttribute( OtherEdgeID, AttributeName, AttributeIndex ) );
-				}
-			}
+			EdgeAttributeList.Attributes.Reset();
+			BackupAllAttributes( EdgeAttributeList, GetMeshDescription()->EdgeAttributes(), OtherEdgeID );
 		}
 
 		// The new edge will be connected to the same polygons as both of the edges we're replacing.  Because
@@ -5623,7 +5299,7 @@ void UEditableMesh::TryToRemoveVertex( const FVertexID VertexID, bool& bOutWasVe
 			EdgeToCreate.VertexID0 = NewEdgeVertexIDs[ 0 ];
 			EdgeToCreate.VertexID1 = NewEdgeVertexIDs[ 1 ];
 			EdgeToCreate.ConnectedPolygons = NewEdgeConnectedPolygons;
-			EdgeToCreate.EdgeAttributes.Attributes = NewEdgeAttributes;
+			EdgeToCreate.EdgeAttributes = EdgeAttributeList;
 
 			static TArray<FEdgeID> NewEdgeIDs;
 			NewEdgeIDs.Reset();
@@ -5679,8 +5355,8 @@ void UEditableMesh::ExtrudePolygons( const TArray<FPolygonID>& PolygonIDs, const
 	static TArray<FAttributesForVertex> AttributesForVertices;
 	AttributesForVertices.Reset();
 
-	static TArray<FVertexAttributesForPolygon> VertexAttributesForPolygons;
-	VertexAttributesForPolygons.Reset();
+	// Get vertex position attribute array
+	const TVertexAttributeArray<FVector>& VertexPositions = GetMeshDescription()->VertexAttributes().GetAttributes<FVector>( MeshAttribute::Vertex::Position );
 
 	// First, let's figure out which of the polygons we were asked to extrude share edges or vertices.  We'll keep those
 	// edges intact!
@@ -5774,9 +5450,10 @@ void UEditableMesh::ExtrudePolygons( const TArray<FPolygonID>& PolygonIDs, const
 					if( !bIsSharedEdge )
 					{
 						// After extruding, all of the edges of the original polygon become hard edges
-						FAttributesForEdge& AttributesForEdge = *new( AttributesForEdges ) FAttributesForEdge();
+						AttributesForEdges.Emplace();
+						FAttributesForEdge& AttributesForEdge = AttributesForEdges.Last();
 						AttributesForEdge.EdgeID = EdgeID;
-						AttributesForEdge.EdgeAttributes.Attributes.Add( FMeshElementAttributeData( UEditableMeshAttribute::EdgeIsHard(), 0, FVector4( 1.0f ) ) );
+						AttributesForEdge.EdgeAttributes.Attributes.Emplace( MeshAttribute::Edge::IsHard, 0, FMeshElementAttributeValue( true ) );
 					}
 
 					FVertexID ExtrudedEdgeVertexIDs[ 2 ];
@@ -5801,7 +5478,7 @@ void UEditableMesh::ExtrudePolygons( const TArray<FPolygonID>& PolygonIDs, const
 
 							ExtrudedEdgeVertexIDPtr = &VertexIDToExtrudedCopy.Add( EdgeVertexID, ExtrudedVertexID );
 							// Push the vertex out along the polygon's normal
-							const FVector OriginalVertexPosition = this->GetVertexAttribute( EdgeVertexID, UEditableMeshAttribute::VertexPosition(), 0 );
+							const FVector OriginalVertexPosition = VertexPositions[ EdgeVertexID ];
 
 							FVector ExtrudedVertexPosition;
 							if( bIsSharedEdge )
@@ -5868,33 +5545,33 @@ void UEditableMesh::ExtrudePolygons( const TArray<FPolygonID>& PolygonIDs, const
 							}
 
 							// Fill in the vertex
-							FAttributesForVertex& AttributesForVertex = *new( AttributesForVertices ) FAttributesForVertex();
+							AttributesForVertices.Emplace();
+							FAttributesForVertex& AttributesForVertex = AttributesForVertices.Last();
 							AttributesForVertex.VertexID = ExtrudedVertexID;
-							AttributesForVertex.VertexAttributes.Attributes.Add( FMeshElementAttributeData(
-								UEditableMeshAttribute::VertexPosition(),
-								0,
-								ExtrudedVertexPosition ) );
+							AttributesForVertex.VertexAttributes.Attributes.Emplace( MeshAttribute::Vertex::Position, 0, FMeshElementAttributeValue( ExtrudedVertexPosition ) );
 						}
 						ExtrudedEdgeVertexIDs[ EdgeVertexNumber ] = *ExtrudedEdgeVertexIDPtr;
 					}
 
 					if( !bIsSharedEdge )
 					{
-						// Add a face that connects the edge to it's extruded counterpart edge
-						FVertexID OriginalVertexIDs[ 4 ];
-
 						static TArray<FVertexAndAttributes> NewSidePolygonVertices;
 						NewSidePolygonVertices.Reset( 4 );
 						NewSidePolygonVertices.SetNum( 4, false );	// Always four edges in an extruded face
 
 						NewSidePolygonVertices[ 0 ].VertexID = EdgeVertexIDs[ 1 ];
-						OriginalVertexIDs[ 0 ] = EdgeVertexIDs[ 1 ];
 						NewSidePolygonVertices[ 1 ].VertexID = EdgeVertexIDs[ 0 ];
-						OriginalVertexIDs[ 1 ] = EdgeVertexIDs[ 0 ];
 						NewSidePolygonVertices[ 2 ].VertexID = ExtrudedEdgeVertexIDs[ 0 ];
-						OriginalVertexIDs[ 2 ] = EdgeVertexIDs[ 0 ];
 						NewSidePolygonVertices[ 3 ].VertexID = ExtrudedEdgeVertexIDs[ 1 ];
-						OriginalVertexIDs[ 3 ] = EdgeVertexIDs[ 1 ];
+
+						// Get vertex instance IDs on this polygon corresponding to the edge start/end vertices
+						const FVertexInstanceID EdgeVertexInstanceID0 = GetMeshDescription()->GetVertexInstanceForPolygonVertex( PolygonID, EdgeVertexIDs[ 0 ] );
+						const FVertexInstanceID EdgeVertexInstanceID1 = GetMeshDescription()->GetVertexInstanceForPolygonVertex( PolygonID, EdgeVertexIDs[ 1 ] );
+
+						BackupAllAttributes( NewSidePolygonVertices[ 0 ].PolygonVertexAttributes, GetMeshDescription()->VertexInstanceAttributes(), EdgeVertexInstanceID1 );
+						BackupAllAttributes( NewSidePolygonVertices[ 1 ].PolygonVertexAttributes, GetMeshDescription()->VertexInstanceAttributes(), EdgeVertexInstanceID0 );
+						BackupAllAttributes( NewSidePolygonVertices[ 2 ].PolygonVertexAttributes, GetMeshDescription()->VertexInstanceAttributes(), EdgeVertexInstanceID0 );
+						BackupAllAttributes( NewSidePolygonVertices[ 3 ].PolygonVertexAttributes, GetMeshDescription()->VertexInstanceAttributes(), EdgeVertexInstanceID1 );
 
 						FPolygonID NewSidePolygonID;	// Filled in below
 						{
@@ -5919,39 +5596,6 @@ void UEditableMesh::ExtrudePolygons( const TArray<FPolygonID>& PolygonIDs, const
 						}
 						AllNewPolygons.Add( NewSidePolygonID );
 
-						// Copy polygon UVs from the original polygon's vertices
-						{
-							FVertexAttributesForPolygon& PolygonNewAttributes = *new( VertexAttributesForPolygons ) FVertexAttributesForPolygon();
-							PolygonNewAttributes.PolygonID = NewSidePolygonID;
-							PolygonNewAttributes.PerimeterVertexAttributeLists.SetNum( NewSidePolygonVertices.Num(), false );
-
-							for( int32 NewVertexIter = 0; NewVertexIter < NewSidePolygonVertices.Num(); ++NewVertexIter )
-							{
-								const int32 PolygonVertexNumber = this->FindPolygonPerimeterVertexNumberForVertex( PolygonID, OriginalVertexIDs[ NewVertexIter ] );
-								check( PolygonVertexNumber != INDEX_NONE );
-
-								TArray<FMeshElementAttributeData>& PerimeterVertexNewAttributes = PolygonNewAttributes.PerimeterVertexAttributeLists[ NewVertexIter ].Attributes;
-								for( int32 TextureCoordinateIndex = 0; TextureCoordinateIndex < TextureCoordinateCount; ++TextureCoordinateIndex )
-								{
-									const FVector4 TextureCoordinate = this->GetPolygonPerimeterVertexAttribute(
-										PolygonID,
-										PolygonVertexNumber,
-										UEditableMeshAttribute::VertexTextureCoordinate(),
-										TextureCoordinateIndex );
-
-									PerimeterVertexNewAttributes.Emplace( UEditableMeshAttribute::VertexTextureCoordinate(), TextureCoordinateIndex, TextureCoordinate );
-								}
-
-								const FVector4 VertexColor = this->GetPolygonPerimeterVertexAttribute(
-									PolygonID,
-									PolygonVertexNumber,
-									UEditableMeshAttribute::VertexColor(),
-									0 );
-
-								PerimeterVertexNewAttributes.Emplace( UEditableMeshAttribute::VertexColor(), 0, VertexColor );
-							}
-						}
-
 						// All of this edges of the new polygon will be hard
 						{
 							const int32 NewPerimeterEdgeCount = this->GetPolygonPerimeterEdgeCount( NewSidePolygonID );
@@ -5963,7 +5607,7 @@ void UEditableMesh::ExtrudePolygons( const TArray<FPolygonID>& PolygonIDs, const
 								// New side polygons get hard edges, but not hard crease weights
 								FAttributesForEdge& AttributesForEdge = *new( AttributesForEdges ) FAttributesForEdge();
 								AttributesForEdge.EdgeID = NewEdgeID;
-								AttributesForEdge.EdgeAttributes.Attributes.Add( FMeshElementAttributeData( UEditableMeshAttribute::EdgeIsHard(), 0, FVector4( 1.0f ) ) );
+								AttributesForEdge.EdgeAttributes.Attributes.Emplace( MeshAttribute::Edge::IsHard, 0, FMeshElementAttributeValue( true ) );
 							}
 						}
 					}
@@ -5987,6 +5631,8 @@ void UEditableMesh::ExtrudePolygons( const TArray<FPolygonID>& PolygonIDs, const
 			NewFrontPolygonVertices.Reset( PolygonVertexIDs.Num() );
 			NewFrontPolygonVertices.SetNum( PolygonVertexIDs.Num(), false );
 
+			const TArray<FVertexInstanceID>& VertexInstanceIDs = GetMeshDescription()->GetPolygonPerimeterVertexInstances( PolygonID );
+
 			// Map all of the polygon's vertex IDs to their extruded counterparts to create the new polygon perimeter
 			for( int32 PolygonVertexNumber = 0; PolygonVertexNumber < PolygonVertexIDs.Num(); ++PolygonVertexNumber )
 			{
@@ -6001,6 +5647,9 @@ void UEditableMesh::ExtrudePolygons( const TArray<FPolygonID>& PolygonIDs, const
 					// We didn't need to extrude a new copy of this vertex (because it was part of a shared edge), so just connect the polygon to the original vertex
 					NewFrontPolygonVertices[ PolygonVertexNumber ].VertexID = VertexID;
 				}
+
+				// Copy vertex instance attributes from original polygon vertex to extruded polygon vertex
+				BackupAllAttributes( NewFrontPolygonVertices[ PolygonVertexNumber ].PolygonVertexAttributes, GetMeshDescription()->VertexInstanceAttributes(), VertexInstanceIDs[ PolygonVertexNumber ] );
 			}
 
 			{
@@ -6024,40 +5673,11 @@ void UEditableMesh::ExtrudePolygons( const TArray<FPolygonID>& PolygonIDs, const
 			}
 			AllNewPolygons.Add( ExtrudedFrontPolygonID );
 
-			// Retain the UVs from the original polygon
-			{
-				FVertexAttributesForPolygon& PolygonNewAttributes = *new( VertexAttributesForPolygons ) FVertexAttributesForPolygon();
-				PolygonNewAttributes.PolygonID = ExtrudedFrontPolygonID;
-				PolygonNewAttributes.PerimeterVertexAttributeLists.SetNum( PolygonVertexIDs.Num(), false );
-
-				for( int32 PolygonVertexNumber = 0; PolygonVertexNumber < PolygonVertexIDs.Num(); ++PolygonVertexNumber )
-				{
-					TArray<FMeshElementAttributeData>& PerimeterVertexNewAttributes = PolygonNewAttributes.PerimeterVertexAttributeLists[ PolygonVertexNumber ].Attributes;
-					for( int32 TextureCoordinateIndex = 0; TextureCoordinateIndex < TextureCoordinateCount; ++TextureCoordinateIndex )
-					{
-						const FVector4 TextureCoordinate = this->GetPolygonPerimeterVertexAttribute(
-							PolygonID,
-							PolygonVertexNumber,
-							UEditableMeshAttribute::VertexTextureCoordinate(),
-							TextureCoordinateIndex );
-
-						FMeshElementAttributeData& PerimeterVertexAttribute = *new( PerimeterVertexNewAttributes ) FMeshElementAttributeData( UEditableMeshAttribute::VertexTextureCoordinate(), TextureCoordinateIndex, TextureCoordinate );
-					}
-
-					const FVector4 VertexColor = this->GetPolygonPerimeterVertexAttribute(
-						PolygonID,
-						PolygonVertexNumber,
-						UEditableMeshAttribute::VertexColor(),
-						0 );
-
-					FMeshElementAttributeData& PerimeterVertexAttribute = *new( PerimeterVertexNewAttributes ) FMeshElementAttributeData( UEditableMeshAttribute::VertexColor(), 0, VertexColor );
-				}
-			}
-
-
 			// All of the border edges of the new polygon will be hard.  If it was a shared edge, then we'll just preserve whatever was
 			// originally going on with the internal edge.
 			{
+				const TEdgeAttributeArray<bool>& EdgeHardnesses = GetMeshDescription()->EdgeAttributes().GetAttributes<bool>( MeshAttribute::Edge::IsHard );
+
 				const int32 NewPerimeterEdgeCount = this->GetPolygonPerimeterEdgeCount( ExtrudedFrontPolygonID );
 				check( NewPerimeterEdgeCount == GetPolygonPerimeterEdgeCount( PolygonID ) );	// New polygon should always have the same number of edges (in the same order) as the original!
 				for( int32 PerimeterEdgeNumber = 0; PerimeterEdgeNumber < NewPerimeterEdgeCount; ++PerimeterEdgeNumber )
@@ -6069,11 +5689,11 @@ void UEditableMesh::ExtrudePolygons( const TArray<FPolygonID>& PolygonIDs, const
 					bool bEdgeWindingIsReversedForPolygon;
 					const FEdgeID EdgeID = this->GetPolygonPerimeterEdge( ExtrudedFrontPolygonID, PerimeterEdgeNumber, /* Out */ bEdgeWindingIsReversedForPolygon );
 
-					const FVector4 NewEdgeHardnessAttribute = bIsSharedEdge ? GetEdgeAttribute( OriginalEdgeID, UEditableMeshAttribute::EdgeIsHard(), 0 ) : FVector4( 1.0f, 0.0f, 0.0f, 0.0f );
+					const bool bNewEdgeHardnessAttribute = bIsSharedEdge ? EdgeHardnesses[ OriginalEdgeID ] : true;
 
 					FAttributesForEdge& AttributesForEdge = *new( AttributesForEdges ) FAttributesForEdge();
 					AttributesForEdge.EdgeID = EdgeID;
-					AttributesForEdge.EdgeAttributes.Attributes.Add( FMeshElementAttributeData( UEditableMeshAttribute::EdgeIsHard(), 0, NewEdgeHardnessAttribute ) );
+					AttributesForEdge.EdgeAttributes.Attributes.Emplace( MeshAttribute::Edge::IsHard, 0, FMeshElementAttributeValue( bNewEdgeHardnessAttribute ) );
 				}
 			}
 		}
@@ -6087,7 +5707,6 @@ void UEditableMesh::ExtrudePolygons( const TArray<FPolygonID>& PolygonIDs, const
 
 	// Update vertex attributes in bulk
 	SetVerticesAttributes( AttributesForVertices );
-	SetPolygonsVertexAttributes( VertexAttributesForPolygons );
 
 	// Delete the original polygons
 	{
@@ -6137,17 +5756,7 @@ void UEditableMesh::ExtendEdges( const TArray<FEdgeID>& EdgeIDs, const bool bWel
 					FVertexToCreate& VertexToCreate = *new( VerticesToCreate ) FVertexToCreate();
 
 					// Copy attributes from the original vertex
-					for( const FName AttributeName : UEditableMesh::GetValidVertexAttributes() )
-					{
-						const int32 MaxAttributeIndex = GetMaxAttributeIndex( AttributeName );
-						for( int32 AttributeIndex = 0; AttributeIndex < MaxAttributeIndex; ++AttributeIndex )
-						{
-							FMeshElementAttributeData& VertexAttribute = *new( VertexToCreate.VertexAttributes.Attributes ) FMeshElementAttributeData(
-								AttributeName,
-								AttributeIndex,
-								GetVertexAttribute( EdgeVertexID, AttributeName, AttributeIndex ) );
-						}
-					}
+					BackupAllAttributes( VertexToCreate.VertexAttributes, GetMeshDescription()->VertexAttributes(), EdgeVertexID );
 
 					// Keep track of which vertex we're creating a counterpart for
 					// @todo mesheditor perf: In the case where bWeldNeighbors=false, we can skip updating this map and instead calculate
@@ -6178,17 +5787,7 @@ void UEditableMesh::ExtendEdges( const TArray<FEdgeID>& EdgeIDs, const bool bWel
 			EdgeToCreate.VertexID1 = NewVertexIDs[ OriginalVertexIDToCreatedVertexNumber.FindChecked( EdgeVertexIDs[ 1 ] ) ];
 
 			// Copy attributes from our original edge
-			for( const FName AttributeName : UEditableMesh::GetValidEdgeAttributes() )
-			{
-				const int32 MaxAttributeIndex = GetMaxAttributeIndex( AttributeName );
-				for( int32 AttributeIndex = 0; AttributeIndex < MaxAttributeIndex; ++AttributeIndex )
-				{
-					FMeshElementAttributeData& EdgeAttribute = *new( EdgeToCreate.EdgeAttributes.Attributes ) FMeshElementAttributeData(
-						AttributeName,
-						AttributeIndex,
-						GetEdgeAttribute( EdgeID, AttributeName, AttributeIndex ) );
-				}
-			}
+			BackupAllAttributes( EdgeToCreate.EdgeAttributes, GetMeshDescription()->EdgeAttributes(), EdgeID );
 
 			// We're not connected to any polygons yet.  That will come later.
 			EdgeToCreate.ConnectedPolygons.Reset();
@@ -6234,44 +5833,24 @@ void UEditableMesh::ExtendEdges( const TArray<FEdgeID>& EdgeIDs, const bool bWel
 			check( PolygonToCreate.PolygonGroupID != FPolygonGroupID::Invalid );
 
 			PolygonToCreate.PerimeterVertices.SetNum( 4, false );
-			FVertexID ConnectedPolygonVertexIDsForTextureCoordinates[ 4 ];
-			{
-				// @todo mesheditor urgent subdiv: This causes degenerate UV triangles, which OpenSubdiv sort of freaks out about (causes flickering geometry)
-				int32 NextVertexNumber = 0;
-				ConnectedPolygonVertexIDsForTextureCoordinates[ NextVertexNumber ] = OriginalEdgeVertexIDs[ 1 ];
-				PolygonToCreate.PerimeterVertices[ NextVertexNumber++ ].VertexID = OriginalEdgeVertexIDs[ 1 ];
 
-				ConnectedPolygonVertexIDsForTextureCoordinates[ NextVertexNumber ] = OriginalEdgeVertexIDs[ 0 ];
-				PolygonToCreate.PerimeterVertices[ NextVertexNumber++ ].VertexID = OriginalEdgeVertexIDs[ 0 ];
-
-				ConnectedPolygonVertexIDsForTextureCoordinates[ NextVertexNumber ] = OriginalEdgeVertexIDs[ 0 ];
-				PolygonToCreate.PerimeterVertices[ NextVertexNumber++ ].VertexID = ExtendedEdgeVertexIDs[ 0 ];
-
-				ConnectedPolygonVertexIDsForTextureCoordinates[ NextVertexNumber ] = OriginalEdgeVertexIDs[ 1 ];
-				PolygonToCreate.PerimeterVertices[ NextVertexNumber++ ].VertexID = ExtendedEdgeVertexIDs[ 1 ];
-			}
+			// @todo mesheditor urgent subdiv: This causes degenerate UV triangles, which OpenSubdiv sort of freaks out about (causes flickering geometry)
+			PolygonToCreate.PerimeterVertices[ 0 ].VertexID = OriginalEdgeVertexIDs[ 1 ];
+			PolygonToCreate.PerimeterVertices[ 1 ].VertexID = OriginalEdgeVertexIDs[ 0 ];
+			PolygonToCreate.PerimeterVertices[ 2 ].VertexID = ExtendedEdgeVertexIDs[ 0 ];
+			PolygonToCreate.PerimeterVertices[ 3 ].VertexID = ExtendedEdgeVertexIDs[ 1 ];
 
 			// Preserve polygon vertex attributes
 			if( ConnectedPolygonID != FPolygonID::Invalid )
 			{
-				for( int32 PerimeterVertexNumber = 0; PerimeterVertexNumber < PolygonToCreate.PerimeterVertices.Num(); ++PerimeterVertexNumber )
-				{
-					const FVertexID ConnectedPolygonVertexIDForTextureCoordinates = ConnectedPolygonVertexIDsForTextureCoordinates[ PerimeterVertexNumber ];
+				// Get vertex instance IDs on this polygon corresponding to the edge start/end vertices
+				const FVertexInstanceID EdgeVertexInstanceID0 = GetMeshDescription()->GetVertexInstanceForPolygonVertex( ConnectedPolygonID, OriginalEdgeVertexIDs[ 0 ] );
+				const FVertexInstanceID EdgeVertexInstanceID1 = GetMeshDescription()->GetVertexInstanceForPolygonVertex( ConnectedPolygonID, OriginalEdgeVertexIDs[ 1 ] );
 
-					TArray<FMeshElementAttributeData>& AttributesForVertex = PolygonToCreate.PerimeterVertices[ PerimeterVertexNumber ].PolygonVertexAttributes.Attributes;
-				
-					const int32 ConnectedPolygonPerimeterVertexNumber = FindPolygonPerimeterVertexNumberForVertex( ConnectedPolygonID, ConnectedPolygonVertexIDForTextureCoordinates );
-					check( ConnectedPolygonPerimeterVertexNumber != INDEX_NONE );
-
-					for( int32 TextureCoordinateIndex = 0; TextureCoordinateIndex < TextureCoordinateCount; ++TextureCoordinateIndex )
-					{
-						const FVector4 TextureCoordinate = GetPolygonPerimeterVertexAttribute( ConnectedPolygonID, ConnectedPolygonPerimeterVertexNumber, UEditableMeshAttribute::VertexTextureCoordinate(), TextureCoordinateIndex );
-						AttributesForVertex.Add( FMeshElementAttributeData( UEditableMeshAttribute::VertexTextureCoordinate(), TextureCoordinateIndex, TextureCoordinate ) );
-					}
-
-					const FVector4 VertexColor = GetPolygonPerimeterVertexAttribute( ConnectedPolygonID, ConnectedPolygonPerimeterVertexNumber, UEditableMeshAttribute::VertexColor(), 0 );
-					AttributesForVertex.Add( FMeshElementAttributeData( UEditableMeshAttribute::VertexColor(), 0, VertexColor ) );
-				}
+				BackupAllAttributes( PolygonToCreate.PerimeterVertices[ 0 ].PolygonVertexAttributes, GetMeshDescription()->VertexInstanceAttributes(), EdgeVertexInstanceID1 );
+				BackupAllAttributes( PolygonToCreate.PerimeterVertices[ 1 ].PolygonVertexAttributes, GetMeshDescription()->VertexInstanceAttributes(), EdgeVertexInstanceID0 );
+				BackupAllAttributes( PolygonToCreate.PerimeterVertices[ 2 ].PolygonVertexAttributes, GetMeshDescription()->VertexInstanceAttributes(), EdgeVertexInstanceID0 );
+				BackupAllAttributes( PolygonToCreate.PerimeterVertices[ 3 ].PolygonVertexAttributes, GetMeshDescription()->VertexInstanceAttributes(), EdgeVertexInstanceID1 );
 			}
 		}
 
@@ -6306,18 +5885,7 @@ void UEditableMesh::ExtendVertices( const TArray<FVertexID>& VertexIDs, const bo
 			const int32 CreatedVertexNumber = VerticesToCreate.Num();
 			FVertexToCreate& VertexToCreate = *new( VerticesToCreate ) FVertexToCreate();
 
-			// Copy attributes from the original vertex
-			for( const FName AttributeName : UEditableMesh::GetValidVertexAttributes() )
-			{
-				const int32 MaxAttributeIndex = GetMaxAttributeIndex( AttributeName );
-				for( int32 AttributeIndex = 0; AttributeIndex < MaxAttributeIndex; ++AttributeIndex )
-				{
-					FMeshElementAttributeData& VertexAttribute = *new( VertexToCreate.VertexAttributes.Attributes ) FMeshElementAttributeData(
-						AttributeName,
-						AttributeIndex,
-						GetVertexAttribute( VertexID, AttributeName, AttributeIndex ) );
-				}
-			}
+			BackupAllAttributes( VertexToCreate.VertexAttributes, GetMeshDescription()->VertexAttributes(), VertexID );
 		}
 
 		CreateVertices( VerticesToCreate, /* Out */ OutNewExtendedVertexIDs );
@@ -6331,6 +5899,8 @@ void UEditableMesh::ExtendVertices( const TArray<FVertexID>& VertexIDs, const bo
 		static TArray<FPolygonToCreate> PolygonsToCreate;
 		PolygonsToCreate.Reset();
 
+		const TVertexAttributeArray<FVector>& VertexPositions = GetMeshDescription()->VertexAttributes().GetAttributes<FVector>( MeshAttribute::Vertex::Position );
+
 		for( int32 VertexNumber = 0; VertexNumber < VertexIDs.Num(); ++VertexNumber )
 		{
 			const FVertexID OriginalVertexID = VertexIDs[ VertexNumber ];
@@ -6351,13 +5921,7 @@ void UEditableMesh::ExtendVertices( const TArray<FVertexID>& VertexIDs, const bo
 					FVertexID EdgeVertexIDs[ 2 ];
 					GetEdgeVertices( ConnectedEdgeID, /* Out */ EdgeVertexIDs[0], /* Out */ EdgeVertexIDs[1] );
 
-					FVector EdgeVertexPositions[ 2 ];
-					for( int32 EdgeVertexNumber = 0; EdgeVertexNumber < 2; ++EdgeVertexNumber )
-					{
-						EdgeVertexPositions[ EdgeVertexNumber ] = GetVertexAttribute( EdgeVertexIDs[ EdgeVertexNumber ], UEditableMeshAttribute::VertexPosition(), 0 );
-					}
-
-					const float SquaredEdgeDistance = FMath::PointDistToSegmentSquared( ReferencePosition, EdgeVertexPositions[ 0 ], EdgeVertexPositions[ 1 ] );
+					const float SquaredEdgeDistance = FMath::PointDistToSegmentSquared( ReferencePosition, VertexPositions[ EdgeVertexIDs[ 0 ] ], VertexPositions[ EdgeVertexIDs[ 1 ] ] );
 					if( SquaredEdgeDistance < ClosestSquaredEdgeDistance )
 					{
 						ClosestVertexID = EdgeVertexIDs[ 0 ] == OriginalVertexID ? EdgeVertexIDs[ 1 ] : EdgeVertexIDs[ 0 ];
@@ -6453,20 +6017,10 @@ void UEditableMesh::ExtendVertices( const TArray<FVertexID>& VertexIDs, const bo
 						for( int32 PerimeterVertexNumber = 0; PerimeterVertexNumber < PolygonToCreate.PerimeterVertices.Num(); ++PerimeterVertexNumber )
 						{
 							const FVertexID ConnectedPolygonVertexIDForTextureCoordinates = ConnectedPolygonVertexIDsForTextureCoordinates[ PerimeterVertexNumber ];
+							const FVertexInstanceID VertexInstanceID = GetMeshDescription()->GetVertexInstanceForPolygonVertex( ConnectedPolygonID, ConnectedPolygonVertexIDForTextureCoordinates );
+							check( VertexInstanceID != FVertexInstanceID::Invalid );
 
-							TArray<FMeshElementAttributeData>& AttributesForVertex = PolygonToCreate.PerimeterVertices[ PerimeterVertexNumber ].PolygonVertexAttributes.Attributes;
-
-							const int32 ConnectedPolygonPerimeterVertexNumber = FindPolygonPerimeterVertexNumberForVertex( ConnectedPolygonID, ConnectedPolygonVertexIDForTextureCoordinates );
-							check( ConnectedPolygonPerimeterVertexNumber != INDEX_NONE );
-
-							for( int32 TextureCoordinateIndex = 0; TextureCoordinateIndex < TextureCoordinateCount; ++TextureCoordinateIndex )
-							{
-								const FVector4 TextureCoordinate = GetPolygonPerimeterVertexAttribute( ConnectedPolygonID, ConnectedPolygonPerimeterVertexNumber, UEditableMeshAttribute::VertexTextureCoordinate(), TextureCoordinateIndex );
-								AttributesForVertex.Add( FMeshElementAttributeData( UEditableMeshAttribute::VertexTextureCoordinate(), TextureCoordinateIndex, TextureCoordinate ) );
-							}
-
-							const FVector4 VertexColor = GetPolygonPerimeterVertexAttribute( ConnectedPolygonID, ConnectedPolygonPerimeterVertexNumber, UEditableMeshAttribute::VertexColor(), 0 );
-							AttributesForVertex.Add( FMeshElementAttributeData( UEditableMeshAttribute::VertexColor(), 0, VertexColor ) );
+							BackupAllAttributes( PolygonToCreate.PerimeterVertices[ PerimeterVertexNumber ].PolygonVertexAttributes, GetMeshDescription()->VertexInstanceAttributes(), VertexInstanceID );
 						}
 					}
 				}
@@ -6524,6 +6078,8 @@ void UEditableMesh::BevelOrInsetPolygons( const TArray<FPolygonID>& PolygonIDs,
 	static TArray<FAttributesForVertex> AttributesForVertices;
 	AttributesForVertices.Reset();
 
+	const TVertexAttributeArray<FVector>& VertexPositions = GetMeshDescription()->VertexAttributes().GetAttributes<FVector>( MeshAttribute::Vertex::Position );
+
 	for( const FPolygonID PolygonID : PolygonIDs )
 	{
 		const FPolygonGroupID PolygonGroupID = GetGroupForPolygon( PolygonID );
@@ -6534,16 +6090,14 @@ void UEditableMesh::BevelOrInsetPolygons( const TArray<FPolygonID>& PolygonIDs,
 		static TArray<FVertexID> PerimeterVertexIDs;
 		GetPolygonPerimeterVertices( PolygonID, /* Out */ PerimeterVertexIDs );
 
+		const TArray<FVertexInstanceID>& PerimeterVertexInstanceIDs = GetMeshDescription()->GetPolygonPerimeterVertexInstances( PolygonID );
+
 		static TArray<FVertexToCreate> VerticesToCreate;
 		VerticesToCreate.Reset( PerimeterVertexIDs.Num() );
 
-		static TArray<TArray<FVector4>> TextureCoordinatesForNewVertices;	// @todo mesheditor perf: Extra memory alloc every time because of nested array
-		TextureCoordinatesForNewVertices.Reset();
-		TextureCoordinatesForNewVertices.Reserve( PerimeterVertexIDs.Num() );
-
-		static TArray<FVector4> VertexColorsForNewVertices;
-		VertexColorsForNewVertices.Reset();
-		VertexColorsForNewVertices.Reserve( PerimeterVertexIDs.Num() );
+		static TArray<TTuple<FMeshTriangle, FVector>> BarycentricWeightsForNewVertices;
+		BarycentricWeightsForNewVertices.Reset();
+		BarycentricWeightsForNewVertices.Reserve( PerimeterVertexIDs.Num() );
 
 		for( int32 PerimeterVertexNumber = 0; PerimeterVertexNumber < PerimeterVertexIDs.Num(); ++PerimeterVertexNumber )
 		{
@@ -6583,7 +6137,7 @@ void UEditableMesh::BevelOrInsetPolygons( const TArray<FPolygonID>& PolygonIDs,
 			}
 
 
-			const FVector VertexPosition = GetVertexAttribute( PerimeterVertexID, UEditableMeshAttribute::VertexPosition(), 0 );
+			const FVector VertexPosition = VertexPositions[ PerimeterVertexID ];
 
 			FVector DirectionTowardCenter;
 			float DistanceToCenter;
@@ -6592,41 +6146,18 @@ void UEditableMesh::BevelOrInsetPolygons( const TArray<FPolygonID>& PolygonIDs,
 			const float InsetOffset = ( DistanceToCenter * InsetProgressTowardCenter + InsetFixedDistance );
 			const FVector InsetVertexPosition = VertexPosition + DirectionTowardCenter * InsetOffset;
 
-			FVertexToCreate& VertexToCreate = *new( VerticesToCreate ) FVertexToCreate();
-			VertexToCreate.VertexAttributes.Attributes.Add( FMeshElementAttributeData( UEditableMeshAttribute::VertexPosition(), 0, InsetVertexPosition ) );
+			VerticesToCreate.Emplace();
+			FVertexToCreate& VertexToCreate = VerticesToCreate.Last();
+			VertexToCreate.VertexAttributes.Attributes.Emplace( MeshAttribute::Vertex::Position, 0, FMeshElementAttributeValue( InsetVertexPosition ) );
 
-			// Interpolate the texture coordinates and vertex colors on this polygon to figure out UVs for our new vertex
+			// @todo mesheditor: interpolate all other vertex attributes (maybe handle this with InterpAllAttributes)
+
+			// Determine the barycentric weights of
 			FMeshTriangle Triangle;
 			FVector TriangleVertexWeights;
 
-			if( ComputeBarycentricWeightForPointOnPolygon( PolygonID, InsetVertexPosition, Triangle, TriangleVertexWeights ) )
-			{
-				TArray<FVector4>& TextureCoordinatesForNewVertex = *new( TextureCoordinatesForNewVertices ) TArray<FVector4>();
-				TextureCoordinatesForNewVertex.SetNum( TextureCoordinateCount, false );
-
-				for( int32 TextureCoordinateIndex = 0; TextureCoordinateIndex < TextureCoordinateCount; ++TextureCoordinateIndex )
-				{
-					TextureCoordinatesForNewVertex[ TextureCoordinateIndex ] =
-						TriangleVertexWeights.X * GetVertexInstanceAttribute( Triangle.VertexInstanceID0, UEditableMeshAttribute::VertexTextureCoordinate(), TextureCoordinateIndex ) +
-						TriangleVertexWeights.Y * GetVertexInstanceAttribute( Triangle.VertexInstanceID1, UEditableMeshAttribute::VertexTextureCoordinate(), TextureCoordinateIndex ) +
-						TriangleVertexWeights.Z * GetVertexInstanceAttribute( Triangle.VertexInstanceID2, UEditableMeshAttribute::VertexTextureCoordinate(), TextureCoordinateIndex );
-				}
-
-				VertexColorsForNewVertices.Add(
-					TriangleVertexWeights.X * GetVertexInstanceAttribute( Triangle.VertexInstanceID0, UEditableMeshAttribute::VertexColor(), 0 ) +
-					TriangleVertexWeights.Y * GetVertexInstanceAttribute( Triangle.VertexInstanceID1, UEditableMeshAttribute::VertexColor(), 0 ) +
-					TriangleVertexWeights.Z * GetVertexInstanceAttribute( Triangle.VertexInstanceID2, UEditableMeshAttribute::VertexColor(), 0 )
-				);
-			}
-			else
-			{
-				TArray<FVector4>& TextureCoordinatesForNewVertex = *new( TextureCoordinatesForNewVertices ) TArray<FVector4>();
-
-				// The new vertex must have been slightly outside the bounds of the polygon, probably just floating point error.
-				// Just use zero'd out texture coordinates and white vertex color in this case.
-				TextureCoordinatesForNewVertex.SetNumZeroed( TextureCoordinateCount );
-				VertexColorsForNewVertices.Emplace( 1.0f, 1.0f, 1.0f, 1.0f );
-			}
+			bool bPointInPolygon = ComputeBarycentricWeightForPointOnPolygon( PolygonID, InsetVertexPosition, Triangle, TriangleVertexWeights );
+			BarycentricWeightsForNewVertices.Emplace( Triangle, TriangleVertexWeights );
 
 			// If we're beveling, go ahead and move the original polygon perimeter vertices
 			if( bShouldBevel )
@@ -6639,13 +6170,10 @@ void UEditableMesh::BevelOrInsetPolygons( const TArray<FPolygonID>& PolygonIDs,
 				const bool bAlreadyHaveVertex = AttributesForVertices.ContainsByPredicate( [PerimeterVertexID]( FAttributesForVertex& AttributesForVertex ) { return AttributesForVertex.VertexID == PerimeterVertexID; } );
 				if( !bAlreadyHaveVertex )
 				{
-					FAttributesForVertex& AttributesForVertex = *new( AttributesForVertices ) FAttributesForVertex();
+					AttributesForVertices.Emplace();
+					FAttributesForVertex& AttributesForVertex = AttributesForVertices.Last();
 					AttributesForVertex.VertexID = PerimeterVertexID;
-
-					FMeshElementAttributeData& ElementAttributeData = *new( AttributesForVertex.VertexAttributes.Attributes ) FMeshElementAttributeData();
-					ElementAttributeData.AttributeName = UEditableMeshAttribute::VertexPosition();
-					ElementAttributeData.AttributeValue = NewVertexPosition;
-					ElementAttributeData.AttributeIndex = 0;
+					AttributesForVertex.VertexAttributes.Attributes.Emplace( MeshAttribute::Vertex::Position, 0, FMeshElementAttributeValue( NewVertexPosition ) );
 				}
 			}
 		}
@@ -6668,91 +6196,58 @@ void UEditableMesh::BevelOrInsetPolygons( const TArray<FPolygonID>& PolygonIDs,
 				const FVertexID RightOriginalVertexID = PerimeterVertexIDs[ RightVertexNumber ];
 				const FVertexID RightInsetVertexID = NewVertexIDs[ RightVertexNumber ];
 
-				FPolygonToCreate& PolygonToCreate = *new( SidePolygonsToCreate ) FPolygonToCreate();
+				const FVertexInstanceID LeftOriginalVertexInstanceID = PerimeterVertexInstanceIDs[ LeftVertexNumber ];
+				const FVertexInstanceID RightOriginalVertexInstanceID = PerimeterVertexInstanceIDs[ RightVertexNumber ];
+
+				SidePolygonsToCreate.Emplace();
+				FPolygonToCreate& PolygonToCreate = SidePolygonsToCreate.Last();
 				PolygonToCreate.PolygonGroupID = PolygonGroupID;
 				if( bShouldBevel )
 				{
 					PolygonToCreate.PolygonEdgeHardness = EPolygonEdgeHardness::AllEdgesHard;
 				}
 
-				// Original left
-				{
-					FVertexAndAttributes& NewPerimeterVertex = *new( PolygonToCreate.PerimeterVertices ) FVertexAndAttributes();
-					NewPerimeterVertex.VertexID = LeftOriginalVertexID;
+				PolygonToCreate.PerimeterVertices.SetNum( 4 );
 
-					// Copy all of the polygon vertex attributes over
-					for( int32 TextureCoordinateIndex = 0; TextureCoordinateIndex < TextureCoordinateCount; ++TextureCoordinateIndex )
-					{
-						NewPerimeterVertex.PolygonVertexAttributes.Attributes.Add( FMeshElementAttributeData(
-							UEditableMeshAttribute::VertexTextureCoordinate(),
-							TextureCoordinateIndex,
-							GetPolygonPerimeterVertexAttribute( PolygonID, LeftVertexNumber, UEditableMeshAttribute::VertexTextureCoordinate(), TextureCoordinateIndex ) ) );
-					}
+				PolygonToCreate.PerimeterVertices[ 0 ].VertexID = LeftOriginalVertexID;
+				PolygonToCreate.PerimeterVertices[ 1 ].VertexID = RightOriginalVertexID;
+				PolygonToCreate.PerimeterVertices[ 2 ].VertexID = RightInsetVertexID;
+				PolygonToCreate.PerimeterVertices[ 3 ].VertexID = LeftInsetVertexID;
 
-					NewPerimeterVertex.PolygonVertexAttributes.Attributes.Add( FMeshElementAttributeData(
-						UEditableMeshAttribute::VertexColor(),
-						0,
-						GetPolygonPerimeterVertexAttribute( PolygonID, LeftVertexNumber, UEditableMeshAttribute::VertexColor(), 0 ) ) );
-				}
+				// Original left
+				BackupAllAttributes( PolygonToCreate.PerimeterVertices[ 0 ].PolygonVertexAttributes, GetMeshDescription()->VertexInstanceAttributes(), LeftOriginalVertexInstanceID );
 
 				// Original right
-				{
-					FVertexAndAttributes& NewPerimeterVertex = *new( PolygonToCreate.PerimeterVertices ) FVertexAndAttributes();
-					NewPerimeterVertex.VertexID = RightOriginalVertexID;
-
-					// Copy all of the polygon vertex attributes over
-					for( int32 TextureCoordinateIndex = 0; TextureCoordinateIndex < TextureCoordinateCount; ++TextureCoordinateIndex )
-					{
-						NewPerimeterVertex.PolygonVertexAttributes.Attributes.Add( FMeshElementAttributeData(
-							UEditableMeshAttribute::VertexTextureCoordinate(),
-							TextureCoordinateIndex,
-							GetPolygonPerimeterVertexAttribute( PolygonID, RightVertexNumber, UEditableMeshAttribute::VertexTextureCoordinate(), TextureCoordinateIndex ) ) );
-					}
-
-					NewPerimeterVertex.PolygonVertexAttributes.Attributes.Add( FMeshElementAttributeData(
-						UEditableMeshAttribute::VertexColor(),
-						0,
-						GetPolygonPerimeterVertexAttribute( PolygonID, RightVertexNumber, UEditableMeshAttribute::VertexColor(), 0 ) ) );
-				}
+				BackupAllAttributes( PolygonToCreate.PerimeterVertices[ 1 ].PolygonVertexAttributes, GetMeshDescription()->VertexInstanceAttributes(), RightOriginalVertexInstanceID );
 
 				// Inset right
+				if( BarycentricWeightsForNewVertices[ RightVertexNumber ].Get<0>().VertexInstanceID0 != FVertexInstanceID::Invalid )
 				{
-					FVertexAndAttributes& NewPerimeterVertex = *new( PolygonToCreate.PerimeterVertices ) FVertexAndAttributes();
-					NewPerimeterVertex.VertexID = RightInsetVertexID;
-
-					// Use interpolated texture coordinates
-					for( int32 TextureCoordinateIndex = 0; TextureCoordinateIndex < TextureCoordinateCount; ++TextureCoordinateIndex )
-					{
-						NewPerimeterVertex.PolygonVertexAttributes.Attributes.Add( FMeshElementAttributeData(
-							UEditableMeshAttribute::VertexTextureCoordinate(),
-							TextureCoordinateIndex,
-							TextureCoordinatesForNewVertices[ RightVertexNumber ][ TextureCoordinateIndex ] ) );
-					}
-
-					NewPerimeterVertex.PolygonVertexAttributes.Attributes.Add( FMeshElementAttributeData(
-						UEditableMeshAttribute::VertexColor(),
-						0,
-						VertexColorsForNewVertices[ RightVertexNumber ] ) );
+					InterpAllAttributes( PolygonToCreate.PerimeterVertices[ 2 ].PolygonVertexAttributes,
+										 GetMeshDescription()->VertexInstanceAttributes(),
+										 BarycentricWeightsForNewVertices[ RightVertexNumber ].Get<0>().VertexInstanceID0,
+										 BarycentricWeightsForNewVertices[ RightVertexNumber ].Get<0>().VertexInstanceID1,
+										 BarycentricWeightsForNewVertices[ RightVertexNumber ].Get<0>().VertexInstanceID2,
+										 BarycentricWeightsForNewVertices[ RightVertexNumber ].Get<1>() );
+				}
+				else
+				{
+					BackupAllAttributes( PolygonToCreate.PerimeterVertices[ 2 ].PolygonVertexAttributes, GetMeshDescription()->VertexInstanceAttributes(), RightOriginalVertexInstanceID );
 				}
 
 				// Inset left
+				if( BarycentricWeightsForNewVertices[ LeftVertexNumber ].Get<0>().VertexInstanceID0 != FVertexInstanceID::Invalid )
 				{
-					FVertexAndAttributes& NewPerimeterVertex = *new( PolygonToCreate.PerimeterVertices ) FVertexAndAttributes();
-					NewPerimeterVertex.VertexID = LeftInsetVertexID;
-
-					// Use interpolated texture coordinates
-					for( int32 TextureCoordinateIndex = 0; TextureCoordinateIndex < TextureCoordinateCount; ++TextureCoordinateIndex )
-					{
-						NewPerimeterVertex.PolygonVertexAttributes.Attributes.Add( FMeshElementAttributeData(
-							UEditableMeshAttribute::VertexTextureCoordinate(),
-							TextureCoordinateIndex,
-							TextureCoordinatesForNewVertices[ LeftVertexNumber ][ TextureCoordinateIndex ] ) );
-					}
-
-					NewPerimeterVertex.PolygonVertexAttributes.Attributes.Add( FMeshElementAttributeData(
-						UEditableMeshAttribute::VertexColor(),
-						0,
-						VertexColorsForNewVertices[ LeftVertexNumber ] ) );
+					InterpAllAttributes( PolygonToCreate.PerimeterVertices[ 3 ].PolygonVertexAttributes,
+										 GetMeshDescription()->VertexInstanceAttributes(),
+										 BarycentricWeightsForNewVertices[ LeftVertexNumber ].Get<0>().VertexInstanceID0,
+										 BarycentricWeightsForNewVertices[ LeftVertexNumber ].Get<0>().VertexInstanceID1,
+										 BarycentricWeightsForNewVertices[ LeftVertexNumber ].Get<0>().VertexInstanceID2,
+										 BarycentricWeightsForNewVertices[ LeftVertexNumber ].Get<1>() );
+				}
+				else
+				{
+					BackupAllAttributes( PolygonToCreate.PerimeterVertices[ 3 ].PolygonVertexAttributes, GetMeshDescription()->VertexInstanceAttributes(), LeftOriginalVertexInstanceID );
 				}
 			}
 		}
@@ -6767,26 +6262,28 @@ void UEditableMesh::BevelOrInsetPolygons( const TArray<FPolygonID>& PolygonIDs,
 				PolygonToCreate.PolygonEdgeHardness = EPolygonEdgeHardness::AllEdgesHard;
 			}
 
+			PolygonToCreate.PerimeterVertices.SetNum( NewVertexIDs.Num() );
+
 			for( int32 NewVertexNumber = 0; NewVertexNumber < NewVertexIDs.Num(); ++NewVertexNumber )
 			{
 				const FVertexID NewVertexID = NewVertexIDs[ NewVertexNumber ];
 
-				FVertexAndAttributes& NewPerimeterVertex = *new( PolygonToCreate.PerimeterVertices ) FVertexAndAttributes();
+				FVertexAndAttributes& NewPerimeterVertex = PolygonToCreate.PerimeterVertices[ NewVertexNumber ];
 				NewPerimeterVertex.VertexID = NewVertexID;
 
-				// Use interpolated texture coordinates
-				for( int32 TextureCoordinateIndex = 0; TextureCoordinateIndex < TextureCoordinateCount; ++TextureCoordinateIndex )
+				if( BarycentricWeightsForNewVertices[ NewVertexNumber ].Get<0>().VertexInstanceID0 != FVertexInstanceID::Invalid )
 				{
-					NewPerimeterVertex.PolygonVertexAttributes.Attributes.Add( FMeshElementAttributeData(
-						UEditableMeshAttribute::VertexTextureCoordinate(),
-						TextureCoordinateIndex,
-						TextureCoordinatesForNewVertices[ NewVertexNumber ][ TextureCoordinateIndex ] ) );
+					InterpAllAttributes( NewPerimeterVertex.PolygonVertexAttributes,
+										 GetMeshDescription()->VertexInstanceAttributes(),
+										 BarycentricWeightsForNewVertices[ NewVertexNumber ].Get<0>().VertexInstanceID0,
+										 BarycentricWeightsForNewVertices[ NewVertexNumber ].Get<0>().VertexInstanceID1,
+										 BarycentricWeightsForNewVertices[ NewVertexNumber ].Get<0>().VertexInstanceID2,
+										 BarycentricWeightsForNewVertices[ NewVertexNumber ].Get<1>() );
+				}
+				else
+				{
+					BackupAllAttributes( NewPerimeterVertex.PolygonVertexAttributes, GetMeshDescription()->VertexInstanceAttributes(), PerimeterVertexInstanceIDs[ NewVertexNumber ] );
 				}
-
-				NewPerimeterVertex.PolygonVertexAttributes.Attributes.Add( FMeshElementAttributeData(
-					UEditableMeshAttribute::VertexColor(),
-					0,
-					VertexColorsForNewVertices[ NewVertexNumber ] ) );
 			}
 		}
 	}
@@ -6795,7 +6292,7 @@ void UEditableMesh::BevelOrInsetPolygons( const TArray<FPolygonID>& PolygonIDs,
 	{
 		const bool bDeleteOrphanedEdges = false;	// No need to delete orphans, because this function won't orphan anything
 		const bool bDeleteOrphanedVertices = false;
-		const bool bDeleteOrphanedVertexInstances = false;
+		const bool bDeleteOrphanedVertexInstances = true;
 		const bool bDeleteEmptyPolygonGroups = false;
 		DeletePolygons( PolygonIDs, bDeleteOrphanedEdges, bDeleteOrphanedVertices, bDeleteOrphanedVertexInstances, bDeleteEmptyPolygonGroups );
 	}
@@ -6846,9 +6343,18 @@ float UEditableMesh::GetPolygonCornerAngleForVertex( const FPolygonID PolygonID,
 	{
 		const int32 NumVertices = Contour.VertexInstanceIDs.Num();
 
-		const FVector PrevVertexPosition = GetVertexInstanceAttribute( Contour.VertexInstanceIDs[ ( ContourIndex + NumVertices - 1 ) % NumVertices ], UEditableMeshAttribute::VertexPosition(), 0 );
-		const FVector ThisVertexPosition = GetVertexInstanceAttribute( Contour.VertexInstanceIDs[ ContourIndex ], UEditableMeshAttribute::VertexPosition(), 0 );
-		const FVector NextVertexPosition = GetVertexInstanceAttribute( Contour.VertexInstanceIDs[ ( ContourIndex + 1 ) % NumVertices ], UEditableMeshAttribute::VertexPosition(), 0 );
+		const int32 PrevIndex = ( ContourIndex + NumVertices - 1 ) % NumVertices;
+		const int32 NextIndex = ( ContourIndex + 1 ) % NumVertices;
+
+		const FVertexID PrevVertexID = GetMeshDescription()->GetVertexInstanceVertex( Contour.VertexInstanceIDs[ PrevIndex ] );
+		const FVertexID ThisVertexID = GetMeshDescription()->GetVertexInstanceVertex( Contour.VertexInstanceIDs[ ContourIndex ] );
+		const FVertexID NextVertexID = GetMeshDescription()->GetVertexInstanceVertex( Contour.VertexInstanceIDs[ NextIndex ] );
+
+		const TVertexAttributeArray<FVector>& VertexPositions = GetMeshDescription()->VertexAttributes().GetAttributes<FVector>( MeshAttribute::Vertex::Position );
+
+		const FVector PrevVertexPosition = VertexPositions[ PrevVertexID ];
+		const FVector ThisVertexPosition = VertexPositions[ ThisVertexID ];
+		const FVector NextVertexPosition = VertexPositions[ NextVertexID ];
 
 		const FVector Direction1 = ( PrevVertexPosition - ThisVertexPosition ).GetSafeNormal();
 		const FVector Direction2 = ( NextVertexPosition - ThisVertexPosition ).GetSafeNormal();
@@ -6890,35 +6396,33 @@ float UEditableMesh::GetPolygonCornerAngleForVertex( const FPolygonID PolygonID,
 
 void UEditableMesh::GeneratePolygonTangentsAndNormals( const TArray<FPolygonID>& PolygonIDs )
 {
-	FPolygonArray& Polygons = GetMeshDescription()->Polygons();
-	const FVertexArray& Vertices = GetMeshDescription()->Vertices();
-	const FVertexInstanceArray& VertexInstances = GetMeshDescription()->VertexInstances();
+	const TVertexAttributeArray<FVector>& VertexPositions = GetMeshDescription()->VertexAttributes().GetAttributes<FVector>( MeshAttribute::Vertex::Position );
+
+	// @todo mesheditor: currently hardcoded to calculate the tangent basis from UV0.
+	const TVertexInstanceAttributeArray<FVector2D>& VertexUVs = GetMeshDescription()->VertexInstanceAttributes().GetAttributes<FVector2D>( MeshAttribute::VertexInstance::TextureCoordinate, 0 );
+
+	TPolygonAttributeArray<FVector>& PolygonNormals = GetMeshDescription()->PolygonAttributes().GetAttributes<FVector>( MeshAttribute::Polygon::Normal );
+	TPolygonAttributeArray<FVector>& PolygonTangents = GetMeshDescription()->PolygonAttributes().GetAttributes<FVector>( MeshAttribute::Polygon::Tangent );
+	TPolygonAttributeArray<FVector>& PolygonBinormals = GetMeshDescription()->PolygonAttributes().GetAttributes<FVector>( MeshAttribute::Polygon::Binormal );
 
 	for( const FPolygonID PolygonID : PolygonIDs )
 	{
-		FMeshPolygon& Polygon = Polygons[ PolygonID ];
-
 		// Calculate the tangent basis for the polygon, based on the average of all constituent triangles
 		FVector Normal = FVector::ZeroVector;
 		FVector Tangent = FVector::ZeroVector;
 		FVector Binormal = FVector::ZeroVector;
 
-		for( const FMeshTriangle& Triangle : Polygon.Triangles )
+		for( const FMeshTriangle& Triangle : GetMeshDescription()->GetPolygonTriangles( PolygonID ) )
 		{
-			const FMeshVertexInstance& VertexInstance0 = VertexInstances[ Triangle.VertexInstanceID0 ];
-			const FMeshVertexInstance& VertexInstance1 = VertexInstances[ Triangle.VertexInstanceID1 ];
-			const FMeshVertexInstance& VertexInstance2 = VertexInstances[ Triangle.VertexInstanceID2 ];
+			const FVertexID VertexID0 = GetMeshDescription()->GetVertexInstanceVertex( Triangle.VertexInstanceID0 );
+			const FVertexID VertexID1 = GetMeshDescription()->GetVertexInstanceVertex( Triangle.VertexInstanceID1 );
+			const FVertexID VertexID2 = GetMeshDescription()->GetVertexInstanceVertex( Triangle.VertexInstanceID2 );
 
-			const FMeshVertex& Vertex0 = Vertices[ VertexInstance0.VertexID ];
-			const FMeshVertex& Vertex1 = Vertices[ VertexInstance1.VertexID ];
-			const FMeshVertex& Vertex2 = Vertices[ VertexInstance2.VertexID ];
+			const FVector DPosition1 = VertexPositions[ VertexID1 ] - VertexPositions[ VertexID0 ];
+			const FVector DPosition2 = VertexPositions[ VertexID2 ] - VertexPositions[ VertexID0 ];
 
-			const FVector DPosition1 = Vertex1.VertexPosition - Vertex0.VertexPosition;
-			const FVector DPosition2 = Vertex2.VertexPosition - Vertex0.VertexPosition;
-
-			// @todo mesheditor: currently hardcoded to calculate the tangent basis from UV0.
-			const FVector2D DUV1 = VertexInstance1.VertexUVs[ 0 ] - VertexInstance0.VertexUVs[ 0 ];
-			const FVector2D DUV2 = VertexInstance2.VertexUVs[ 0 ] - VertexInstance0.VertexUVs[ 0 ];
+			const FVector2D DUV1 = VertexUVs[ Triangle.VertexInstanceID1 ] - VertexUVs[ Triangle.VertexInstanceID0 ];
+			const FVector2D DUV2 = VertexUVs[ Triangle.VertexInstanceID2 ] - VertexUVs[ Triangle.VertexInstanceID0 ];
 
 			// We have a left-handed coordinate system, but a counter-clockwise winding order
 			// Hence normal calculation has to take the triangle vectors cross product in reverse.
@@ -6932,64 +6436,60 @@ void UEditableMesh::GeneratePolygonTangentsAndNormals( const TArray<FPolygonID>&
 			Binormal += ( DPosition2 * DUV1.X - DPosition1 * DUV2.X ) * InvDetUV;
 		}
 
-		Polygon.PolygonNormal = Normal.GetSafeNormal();
-		Polygon.PolygonTangent = Tangent.GetSafeNormal();
-		Polygon.PolygonBinormal = Binormal.GetSafeNormal();
+		PolygonNormals[ PolygonID ] = Normal.GetSafeNormal();
+		PolygonTangents[ PolygonID ] = Tangent.GetSafeNormal();
+		PolygonBinormals[ PolygonID ] = Binormal.GetSafeNormal();
 	}
 }
 
 
 void UEditableMesh::GenerateTangentsAndNormals()
 {
+	const TPolygonAttributeArray<FVector>& PolygonNormals = GetMeshDescription()->PolygonAttributes().GetAttributes<FVector>( MeshAttribute::Polygon::Normal );
+	const TPolygonAttributeArray<FVector>& PolygonTangents = GetMeshDescription()->PolygonAttributes().GetAttributes<FVector>( MeshAttribute::Polygon::Tangent );
+	const TPolygonAttributeArray<FVector>& PolygonBinormals = GetMeshDescription()->PolygonAttributes().GetAttributes<FVector>( MeshAttribute::Polygon::Binormal );
+
 	GeneratePolygonTangentsAndNormals( PolygonsPendingNewTangentBasis.Array() );
 
 	static TSet<FVertexInstanceID> VertexInstanceIDs;
 	VertexInstanceIDs.Reset();
 
-	const FPolygonArray& Polygons = GetMeshDescription()->Polygons();
 	for( const FPolygonID PolygonID : PolygonsPendingNewTangentBasis )
 	{
-		const FMeshPolygon& Polygon = Polygons[ PolygonID ];
-
-		VertexInstanceIDs.Append( Polygon.PerimeterContour.VertexInstanceIDs );
-		for( const FMeshPolygonContour& HoleContour : Polygon.HoleContours )
-		{
-			VertexInstanceIDs.Append( HoleContour.VertexInstanceIDs );
-		}
+		VertexInstanceIDs.Append( GetMeshDescription()->GetPolygonPerimeterVertexInstances( PolygonID ) );
+		// @todo holes
 	}
 
 	static TArray<FAttributesForVertexInstance> AttributesForVertexInstances;
 	AttributesForVertexInstances.Reset( VertexInstanceIDs.Num() );
 
-	FVertexInstanceArray& VertexInstances = GetMeshDescription()->VertexInstances();
 	for( const FVertexInstanceID VertexInstanceID : VertexInstanceIDs )
 	{
 		FVector Normal = FVector::ZeroVector;
 		FVector Tangent = FVector::ZeroVector;
 		FVector Binormal = FVector::ZeroVector;
 
-		FMeshVertexInstance& VertexInstance = VertexInstances[ VertexInstanceID ];
-		const FVertexID VertexID = VertexInstance.VertexID;
+		const FVertexID VertexID = GetMeshDescription()->GetVertexInstanceVertex( VertexInstanceID );
 
 		// Get all polygons connected to this vertex instance, and also any in the same smoothing group connected to a different vertex instance
 		// (as they still have influence over the normal).
 		static TArray<FPolygonID> AllConnectedPolygons;
-		check( VertexInstance.ConnectedPolygons.Num() > 0 );
-		GetPolygonsInSameSoftEdgedGroup( VertexID, VertexInstance.ConnectedPolygons[ 0 ], AllConnectedPolygons );
+		const TArray<FPolygonID>& VertexInstanceConnectedPolygons = GetMeshDescription()->GetVertexInstanceConnectedPolygons( VertexInstanceID );
+		check( VertexInstanceConnectedPolygons.Num() > 0 );
+		GetPolygonsInSameSoftEdgedGroup( VertexID, VertexInstanceConnectedPolygons[ 0 ], AllConnectedPolygons );
 
 		// The vertex instance normal is computed as a sum of all connected polygons' normals, weighted by the angle they make with the vertex
 		for( const FPolygonID ConnectedPolygonID : AllConnectedPolygons )
 		{
-			float Angle = GetPolygonCornerAngleForVertex( ConnectedPolygonID, VertexID );
+			const float Angle = GetPolygonCornerAngleForVertex( ConnectedPolygonID, VertexID );
 
-			const FMeshPolygon& Polygon = Polygons[ ConnectedPolygonID ];
-			Normal += Polygon.PolygonNormal * Angle;
+			Normal += PolygonNormals[ ConnectedPolygonID ] * Angle;
 
 			// If this polygon is actually connected to the vertex instance we're processing, also include its contributions towards the tangent
-			if( VertexInstance.ConnectedPolygons.Contains( ConnectedPolygonID ) )
+			if( VertexInstanceConnectedPolygons.Contains( ConnectedPolygonID ) )
 			{
-				Tangent += Polygon.PolygonTangent * Angle;
-				Binormal += Polygon.PolygonBinormal * Angle;
+				Tangent += PolygonTangents[ ConnectedPolygonID ] * Angle;
+				Binormal += PolygonBinormals[ ConnectedPolygonID ] * Angle;
 			}
 		}
 
@@ -7014,9 +6514,9 @@ void UEditableMesh::GenerateTangentsAndNormals()
 
 		AttributesForVertexInstance.VertexInstanceID = VertexInstanceID;
 		AttributesForVertexInstance.VertexInstanceAttributes.Attributes.Reset( 3 );
-		AttributesForVertexInstance.VertexInstanceAttributes.Attributes.Emplace( UEditableMeshAttribute::VertexNormal(), 0, FVector4( Normal ) );
-		AttributesForVertexInstance.VertexInstanceAttributes.Attributes.Emplace( UEditableMeshAttribute::VertexTangent(), 0, FVector4( Tangent ) );
-		AttributesForVertexInstance.VertexInstanceAttributes.Attributes.Emplace( UEditableMeshAttribute::VertexBinormalSign(), 0, FVector4( BinormalSign, 0.0f, 0.0f, 0.0f ) );
+		AttributesForVertexInstance.VertexInstanceAttributes.Attributes.Emplace( MeshAttribute::VertexInstance::Normal, 0, FMeshElementAttributeValue( Normal ) );
+		AttributesForVertexInstance.VertexInstanceAttributes.Attributes.Emplace( MeshAttribute::VertexInstance::Tangent, 0, FMeshElementAttributeValue( Tangent ) );
+		AttributesForVertexInstance.VertexInstanceAttributes.Attributes.Emplace( MeshAttribute::VertexInstance::BinormalSign, 0, FMeshElementAttributeValue( BinormalSign ) );
 	}
 
 	SetVertexInstancesAttributes( AttributesForVertexInstances );
@@ -7077,7 +6577,7 @@ void UEditableMesh::GenerateTangentsForPolygons( const TArray< FPolygonID >& Pol
 			const FMikkUserData& UserData = *static_cast<FMikkUserData*>( Context->m_pUserData );
 
 			const FPolygonID PolygonID = UserData.Polygons[ MikkFaceIndex ];
-			const FVector VertexPosition = UserData.Self->GetPolygonPerimeterVertexAttribute( PolygonID, MikkVertexIndex, UEditableMeshAttribute::VertexPosition(), 0 );
+			const FVector VertexPosition = UserData.Self->GetPolygonPerimeterVertexAttribute( PolygonID, MikkVertexIndex, MeshAttribute::Vertex::Position, 0 );
 
 			OutPosition[0] = VertexPosition.X;
 			OutPosition[1] = VertexPosition.Y;
@@ -7089,7 +6589,7 @@ void UEditableMesh::GenerateTangentsForPolygons( const TArray< FPolygonID >& Pol
 			const FMikkUserData& UserData = *static_cast<FMikkUserData*>( Context->m_pUserData );
 
 			const FPolygonID PolygonID = UserData.Polygons[ MikkFaceIndex ];
-			const FVector PolygonVertexNormal = UserData.Self->GetPolygonPerimeterVertexAttribute( PolygonID, MikkVertexIndex, UEditableMeshAttribute::VertexNormal(), 0 );
+			const FVector PolygonVertexNormal = UserData.Self->GetPolygonPerimeterVertexAttribute( PolygonID, MikkVertexIndex, MeshAttribute::VertexInstance::Normal, 0 );
 
 			OutNormal[ 0 ] = PolygonVertexNormal.X;
 			OutNormal[ 1 ] = PolygonVertexNormal.Y;
@@ -7104,7 +6604,7 @@ void UEditableMesh::GenerateTangentsForPolygons( const TArray< FPolygonID >& Pol
 			const int32 TextureCoordinateIndex = 0;
 
 			const FPolygonID PolygonID = UserData.Polygons[ MikkFaceIndex ];
-			const FVector2D PolygonVertexTextureCoordinate( UserData.Self->GetPolygonPerimeterVertexAttribute( PolygonID, MikkVertexIndex, UEditableMeshAttribute::VertexTextureCoordinate(), TextureCoordinateIndex ) );
+			const FVector2D PolygonVertexTextureCoordinate( UserData.Self->GetPolygonPerimeterVertexAttribute( PolygonID, MikkVertexIndex, MeshAttribute::VertexInstance::TextureCoordinate, TextureCoordinateIndex ) );
 
 			OutUV[0] = PolygonVertexTextureCoordinate.X;
 			OutUV[1] = PolygonVertexTextureCoordinate.Y;
@@ -7119,8 +6619,8 @@ void UEditableMesh::GenerateTangentsForPolygons( const TArray< FPolygonID >& Pol
 			// Save it!
 			TArray<FMeshElementAttributeData>& PerimeterVertexNewAttributes = UserData.VertexAttributesForPolygons[ MikkFaceIndex ].PerimeterVertexAttributeLists[ MikkVertexIndex ].Attributes;
 			check( PerimeterVertexNewAttributes.Num() == 0 );
-			PerimeterVertexNewAttributes.Emplace( UEditableMeshAttribute::VertexTangent(), 0, FVector4( NewTangent, 0.0f ) );
-			PerimeterVertexNewAttributes.Emplace( UEditableMeshAttribute::VertexBinormalSign(), 0, FVector4( BitangentSign ) );
+			PerimeterVertexNewAttributes.Emplace( MeshAttribute::VertexInstance::Tangent, 0, FVector4( NewTangent, 0.0f ) );
+			PerimeterVertexNewAttributes.Emplace( MeshAttribute::VertexInstance::BinormalSign, 0, FVector4( BitangentSign ) );
 		}
 	};
 
@@ -7164,12 +6664,10 @@ void UEditableMesh::SetVerticesCornerSharpness( const TArray<FVertexID>& VertexI
 	{
 		const FVertexID VertexID = VertexIDs[ VertexNumber ];
 
-		FAttributesForVertex& AttributesForVertex = *new( AttributesForVertices ) FAttributesForVertex();
+		AttributesForVertices.Emplace();
+		FAttributesForVertex& AttributesForVertex = AttributesForVertices.Last();
 		AttributesForVertex.VertexID = VertexID;
-		AttributesForVertex.VertexAttributes.Attributes.Add( FMeshElementAttributeData(
-			UEditableMeshAttribute::VertexCornerSharpness(),
-			0,
-			FVector4( VerticesNewSharpness[ VertexNumber ] ) ) );
+		AttributesForVertex.VertexAttributes.Attributes.Emplace( MeshAttribute::Vertex::CornerSharpness, 0, FMeshElementAttributeValue( VerticesNewSharpness[ VertexNumber ] ) );
 	}
 
 	SetVerticesAttributes( AttributesForVertices );
@@ -7187,12 +6685,10 @@ void UEditableMesh::SetEdgesCreaseSharpness( const TArray<FEdgeID>& EdgeIDs, con
 	{
 		const FEdgeID EdgeID = EdgeIDs[ EdgeNumber ];
 
-		FAttributesForEdge& AttributesForEdge = *new( AttributesForEdges ) FAttributesForEdge();
+		AttributesForEdges.Emplace();
+		FAttributesForEdge& AttributesForEdge = AttributesForEdges.Last();
 		AttributesForEdge.EdgeID = EdgeID;
-		AttributesForEdge.EdgeAttributes.Attributes.Add( FMeshElementAttributeData(
-			UEditableMeshAttribute::EdgeCreaseSharpness(),
-			0,
-			FVector4( EdgesNewCreaseSharpness[ EdgeNumber ] ) ) );
+		AttributesForEdge.EdgeAttributes.Attributes.Emplace( MeshAttribute::Edge::CreaseSharpness, 0, FMeshElementAttributeValue( EdgesNewCreaseSharpness[ EdgeNumber ] ) );
 	}
 
 	SetEdgesAttributes( AttributesForEdges );
@@ -7213,12 +6709,10 @@ void UEditableMesh::SetEdgesHardness( const TArray<FEdgeID>& EdgeIDs, const TArr
 	{
 		const FEdgeID EdgeID = EdgeIDs[ EdgeNumber ];
 
-		FAttributesForEdge& AttributesForEdge = *new( AttributesForEdges ) FAttributesForEdge();
+		AttributesForEdges.Emplace();
+		FAttributesForEdge& AttributesForEdge = AttributesForEdges.Last();
 		AttributesForEdge.EdgeID = EdgeID;
-		AttributesForEdge.EdgeAttributes.Attributes.Add( FMeshElementAttributeData(
-			UEditableMeshAttribute::EdgeIsHard(),
-			0,
-			FVector4( EdgesNewIsHard[ EdgeNumber ] ? 1.0f : 0.0f ) ) );
+		AttributesForEdge.EdgeAttributes.Attributes.Emplace( MeshAttribute::Edge::IsHard, 0, FMeshElementAttributeValue( EdgesNewIsHard[ EdgeNumber ] ) );
 
 		// Get the polygons this edge is connected to.  They'll need new normals.
 		static TArray<FPolygonID> ConnectedPolygonIDs;
@@ -7573,27 +7067,17 @@ void UEditableMesh::FlipPolygons( const TArray<FPolygonID>& PolygonIDs )
 		PolygonToCreate.OriginalPolygonID = OriginalPolygonID;
 
 		// Iterate backwards to add the vertices for the polygon, because we're flipping it over.
-		const int32 PerimeterVertexCount = GetPolygonPerimeterVertexCount( OriginalPolygonID );
+		const TArray<FVertexInstanceID>& PerimeterVertexInstanceIDs = GetMeshDescription()->GetPolygonPerimeterVertexInstances( OriginalPolygonID );
+		const int32 PerimeterVertexCount = PerimeterVertexInstanceIDs.Num();
+
 		PolygonToCreate.PerimeterVertices.Reserve( PerimeterVertexCount );
 		for( int32 VertexNumber = PerimeterVertexCount - 1; VertexNumber >= 0; --VertexNumber )
 		{
-			FVertexAndAttributes& PerimeterVertex = *new( PolygonToCreate.PerimeterVertices ) FVertexAndAttributes();
-
-			PerimeterVertex.VertexID = GetPolygonPerimeterVertex( OriginalPolygonID, VertexNumber );
+			PolygonToCreate.PerimeterVertices.Emplace();
+			FVertexAndAttributes& PerimeterVertex = PolygonToCreate.PerimeterVertices.Last();
 
-			// Save off information about the polygons, so we can re-create them with reverse winding.
-			TArray<FMeshElementAttributeData>& PerimeterVertexAttributeList = PerimeterVertex.PolygonVertexAttributes.Attributes;
-			for( const FName AttributeName : UEditableMesh::GetValidVertexInstanceAttributes() )
-			{
-				const int32 MaxAttributeIndex = GetMaxAttributeIndex( AttributeName );
-				for( int32 AttributeIndex = 0; AttributeIndex < MaxAttributeIndex; ++AttributeIndex )
-				{
-					FMeshElementAttributeData& PolygonVertexAttribute = *new( PerimeterVertexAttributeList ) FMeshElementAttributeData(
-						AttributeName,
-						AttributeIndex,
-						GetPolygonPerimeterVertexAttribute( OriginalPolygonID, VertexNumber, AttributeName, AttributeIndex ) );
-				}
-			}
+			// Just create a polygon from existing vertex instance IDs
+			PerimeterVertex.VertexInstanceID = PerimeterVertexInstanceIDs[ VertexNumber ];
 		}
 
 		// @todo mesheditor hole: Perserve holes in this polygon!  (Vertex IDs and Attributes!) Should these be flipped too?  Probably.
@@ -7866,6 +7350,8 @@ void UEditableMesh::WeldVertices( const TArray<FVertexID>& VertexIDsToWeld, FVer
 		return;
 	}
 
+	const TVertexAttributeArray<FVector>& VertexPositions = GetMeshDescription()->VertexAttributes().GetAttributes<FVector>( MeshAttribute::Vertex::Position );
+
 	// Create new welded vertex
 	static TArray<FVertexID> NewVertices;
 	{
@@ -7873,13 +7359,11 @@ void UEditableMesh::WeldVertices( const TArray<FVertexID>& VertexIDsToWeld, FVer
 		VerticesToCreate.SetNum( 1 );
 		FVertexToCreate& VertexToCreate = VerticesToCreate.Last();
 
+		const FVector NewVertexPosition = VertexPositions[ VertexIDsToWeld.Last() ];
+
 		// The vertex which is created will be at the position of the last vertex in the array of vertices to weld.
 		// @todo mesheditor: maybe specify a position to the method instead of using one of the existing vertices?
-		VertexToCreate.VertexAttributes.Attributes.Emplace(
-			UEditableMeshAttribute::VertexPosition(),
-			0,
-			GetVertexAttribute( VertexIDsToWeld.Last(), UEditableMeshAttribute::VertexPosition(), 0 )
-		);
+		VertexToCreate.VertexAttributes.Attributes.Emplace( MeshAttribute::Vertex::Position, 0, FMeshElementAttributeValue( NewVertexPosition ) );
 
 		// @todo mesheditor: vertex corner sharpness too?
 
@@ -7903,6 +7387,7 @@ void UEditableMesh::WeldVertices( const TArray<FVertexID>& VertexIDsToWeld, FVer
 		static TArray<FEdgeID> PolygonEdgeIDs;
 		GetPolygonPerimeterVertices( ConnectedPolygonID, PolygonVertexIDs );
 		GetPolygonPerimeterEdges( ConnectedPolygonID, PolygonEdgeIDs );
+		const TArray<FVertexInstanceID>& PolygonVertexInstanceIDs = GetMeshDescription()->GetPolygonPerimeterVertexInstances( ConnectedPolygonID );
 
 		// Get the index range of perimeter vertices to be welded.
 		// This should definitely be valid, as any invalid welded poly will have caused early exit, above.
@@ -7961,36 +7446,13 @@ void UEditableMesh::WeldVertices( const TArray<FVertexID>& VertexIDsToWeld, FVer
 			}
 
 			// Copy the polygon vertex attributes over
-			for( int32 TextureCoordinateIndex = 0; TextureCoordinateIndex < TextureCoordinateCount; ++TextureCoordinateIndex )
-			{
-				VertexAndAttributes.PolygonVertexAttributes.Attributes.Add( FMeshElementAttributeData(
-					UEditableMeshAttribute::VertexTextureCoordinate(),
-					TextureCoordinateIndex,
-					GetPolygonPerimeterVertexAttribute( ConnectedPolygonID, Index, UEditableMeshAttribute::VertexTextureCoordinate(), TextureCoordinateIndex ) )
-				);
-			}
-
-			VertexAndAttributes.PolygonVertexAttributes.Attributes.Add( FMeshElementAttributeData(
-				UEditableMeshAttribute::VertexColor(),
-				0,
-				GetPolygonPerimeterVertexAttribute( ConnectedPolygonID, Index, UEditableMeshAttribute::VertexColor(), 0 ) )
-			);
+			BackupAllAttributes( VertexAndAttributes.PolygonVertexAttributes, GetMeshDescription()->VertexInstanceAttributes(), PolygonVertexInstanceIDs[ Index ] );
 
 			// Prepare to assign the old edge's attributes to the new edge.
 			// We build up an array of edge attributes to set, in perimeter vertex order for each polygon.
 			AttributesForEdges.Emplace();
 			FAttributesForEdge& AttributesForEdge = AttributesForEdges.Last();
-
-			AttributesForEdge.EdgeAttributes.Attributes.Add( FMeshElementAttributeData(
-				UEditableMeshAttribute::EdgeIsHard(),
-				0,
-				GetEdgeAttribute( PolygonEdgeIDs[ Index ], UEditableMeshAttribute::EdgeIsHard(), 0 ) )
-			);
-			AttributesForEdge.EdgeAttributes.Attributes.Add( FMeshElementAttributeData( 
-				UEditableMeshAttribute::EdgeCreaseSharpness(), 
-				0, 
-				GetEdgeAttribute( PolygonEdgeIDs[ Index ], UEditableMeshAttribute::EdgeCreaseSharpness(), 0 ) )
-			);
+			BackupAllAttributes( AttributesForEdge.EdgeAttributes, GetMeshDescription()->EdgeAttributes(), PolygonEdgeIDs[ Index ] );
 		}
 	}
 
@@ -8037,6 +7499,8 @@ void UEditableMesh::TessellatePolygons( const TArray<FPolygonID>& PolygonIDs, co
 {
 	OutNewPolygonIDs.Reset();
 
+	const TVertexAttributeArray<FVector>& VertexPositions = GetMeshDescription()->VertexAttributes().GetAttributes<FVector>( MeshAttribute::Vertex::Position );
+
 	//
 	// Simple tessellation algorithm:
 	//
@@ -8065,7 +7529,7 @@ void UEditableMesh::TessellatePolygons( const TArray<FPolygonID>& PolygonIDs, co
 				const FVector PolygonCenter = ComputePolygonCenter( PolygonID );
 
 				FVertexToCreate& VertexToCreate = *new( VerticesToCreate ) FVertexToCreate();
-				VertexToCreate.VertexAttributes.Attributes.Add( FMeshElementAttributeData( UEditableMeshAttribute::VertexPosition(), 0, PolygonCenter ) );
+				VertexToCreate.VertexAttributes.Attributes.Emplace( MeshAttribute::Vertex::Position, 0, FMeshElementAttributeValue( PolygonCenter ) );
 			}
 		}
 
@@ -8134,6 +7598,8 @@ void UEditableMesh::TessellatePolygons( const TArray<FPolygonID>& PolygonIDs, co
 			static TArray<FVertexID> PerimeterVertexIDs;
 			GetPolygonPerimeterVertices( PolygonID, /* Out */ PerimeterVertexIDs );
 
+			const TArray<FVertexInstanceID>& PerimeterVertexInstanceIDs = GetMeshDescription()->GetPolygonPerimeterVertexInstances( PolygonID );
+
 			const int32 PerimeterVertexCount = PerimeterEdgeCount;
 			const int32 OriginalPerimeterEdgeCount = PerimeterEdgeCount / 2;
 			for( int32 OriginalPerimeterEdgeNumber = 0; OriginalPerimeterEdgeNumber < OriginalPerimeterEdgeCount; ++OriginalPerimeterEdgeNumber )
@@ -8147,7 +7613,8 @@ void UEditableMesh::TessellatePolygons( const TArray<FPolygonID>& PolygonIDs, co
 
 				for( int32 QuadVertexNumber = 0; QuadVertexNumber < 4; ++QuadVertexNumber )
 				{
-					FVertexAndAttributes& VertexAndAttributes = *new( PolygonToCreate.PerimeterVertices ) FVertexAndAttributes();
+					PolygonToCreate.PerimeterVertices.Emplace();
+					FVertexAndAttributes& VertexAndAttributes = PolygonToCreate.PerimeterVertices.Last();
 
 					int32 PerimeterVertexNumber = INDEX_NONE;
 					switch( QuadVertexNumber )
@@ -8179,50 +7646,21 @@ void UEditableMesh::TessellatePolygons( const TArray<FPolygonID>& PolygonIDs, co
 
 						// Generate interpolated UVs and vertex colors for the new vertex in the center
 						{
-							const FVector CenterVertexPosition = GetVertexAttribute( PolygonCenterVertexID, UEditableMeshAttribute::VertexPosition(), 0 );
+							const FVector CenterVertexPosition = VertexPositions[ PolygonCenterVertexID ];
 
 							FMeshTriangle Triangle;
 							FVector TriangleVertexWeights;
 							if( ComputeBarycentricWeightForPointOnPolygon( PolygonID, CenterVertexPosition, /* Out */ Triangle, /* Out */ TriangleVertexWeights ) )
 							{
-								for( int32 TextureCoordinateIndex = 0; TextureCoordinateIndex < TextureCoordinateCount; ++TextureCoordinateIndex )
-								{
-									VertexAndAttributes.PolygonVertexAttributes.Attributes.Add( FMeshElementAttributeData(
-										UEditableMeshAttribute::VertexTextureCoordinate(),
-										TextureCoordinateIndex,
-										TriangleVertexWeights.X * GetVertexInstanceAttribute( Triangle.VertexInstanceID0, UEditableMeshAttribute::VertexTextureCoordinate(), TextureCoordinateIndex ) +
-										TriangleVertexWeights.Y * GetVertexInstanceAttribute( Triangle.VertexInstanceID1, UEditableMeshAttribute::VertexTextureCoordinate(), TextureCoordinateIndex ) +
-										TriangleVertexWeights.Z * GetVertexInstanceAttribute( Triangle.VertexInstanceID2, UEditableMeshAttribute::VertexTextureCoordinate(), TextureCoordinateIndex ) ) );
-								}
-
-								VertexAndAttributes.PolygonVertexAttributes.Attributes.Add( FMeshElementAttributeData(
-									UEditableMeshAttribute::VertexColor(),
-									0,
-									TriangleVertexWeights.X * GetVertexInstanceAttribute( Triangle.VertexInstanceID0, UEditableMeshAttribute::VertexColor(), 0 ) +
-									TriangleVertexWeights.Y * GetVertexInstanceAttribute( Triangle.VertexInstanceID1, UEditableMeshAttribute::VertexColor(), 0 ) +
-									TriangleVertexWeights.Z * GetVertexInstanceAttribute( Triangle.VertexInstanceID2, UEditableMeshAttribute::VertexColor(), 0 ) ) );
+								InterpAllAttributes( VertexAndAttributes.PolygonVertexAttributes, GetMeshDescription()->VertexInstanceAttributes(),
+									Triangle.VertexInstanceID0, Triangle.VertexInstanceID1, Triangle.VertexInstanceID2, TriangleVertexWeights );
 							}
 						}
 					}
 					else
 					{
-						VertexAndAttributes.VertexID = PerimeterVertexIDs[ PerimeterVertexNumber ];
-
-						// Transfer over UVs and vertex colors from the original polygon that we're replacing
-						{
-							for( int32 TextureCoordinateIndex = 0; TextureCoordinateIndex < TextureCoordinateCount; ++TextureCoordinateIndex )
-							{
-								VertexAndAttributes.PolygonVertexAttributes.Attributes.Add( FMeshElementAttributeData(
-									UEditableMeshAttribute::VertexTextureCoordinate(),
-									TextureCoordinateIndex,
-									GetPolygonPerimeterVertexAttribute( PolygonID, PerimeterVertexNumber, UEditableMeshAttribute::VertexTextureCoordinate(), TextureCoordinateIndex ) ) );
-							}
-
-							VertexAndAttributes.PolygonVertexAttributes.Attributes.Add( FMeshElementAttributeData(
-								UEditableMeshAttribute::VertexColor(),
-								0,
-								GetPolygonPerimeterVertexAttribute( PolygonID, PerimeterVertexNumber, UEditableMeshAttribute::VertexColor(), 0 ) ) );
-						}
+						// Reuse the original vertex instance
+						VertexAndAttributes.VertexInstanceID = PerimeterVertexInstanceIDs[ PerimeterVertexNumber ];
 					}
 				}
 			}
@@ -8231,88 +7669,38 @@ void UEditableMesh::TessellatePolygons( const TArray<FPolygonID>& PolygonIDs, co
 		{
 			if( TriangleTessellationMode == ETriangleTessellationMode::ThreeTriangles )
 			{
+				const TArray<FVertexInstanceID>& PerimeterVertexInstanceIDs = GetMeshDescription()->GetPolygonPerimeterVertexInstances( PolygonID );
+				check( PerimeterVertexInstanceIDs.Num() == 3 );
+
 				// Define the three new triangles for the original tessellated triangle
-				for( int32 PerimeterEdgeNumber = 0; PerimeterEdgeNumber < 3; ++PerimeterEdgeNumber )
+				for( int32 PerimeterVertexNumber = 0; PerimeterVertexNumber < 3; ++PerimeterVertexNumber )
 				{
-					bool bIsEdgeWindingReversedForPolygon;
-					const FEdgeID EdgeID = GetPolygonPerimeterEdge( PolygonID, PerimeterEdgeNumber, /* Out */ bIsEdgeWindingReversedForPolygon );
+					const FVertexInstanceID VertexInstanceID0 = PerimeterVertexInstanceIDs[ PerimeterVertexNumber ];
+					const FVertexInstanceID VertexInstanceID1 = PerimeterVertexInstanceIDs[ ( PerimeterVertexNumber + 1 ) % 3 ];
 
-					FVertexID EdgeVertexID0, EdgeVertexID1;
-					GetEdgeVertices( EdgeID, /* Out */ EdgeVertexID0, /* Out */ EdgeVertexID1 );
+					PolygonsToCreate.Emplace();
+					FPolygonToCreate& PolygonToCreate = PolygonsToCreate.Last();
 
-					FPolygonToCreate& PolygonToCreate = *new( PolygonsToCreate ) FPolygonToCreate();
 					PolygonToCreate.PolygonGroupID = PolygonGroupID;
+					PolygonToCreate.PerimeterVertices.SetNum( 3 );
 
-					for( int32 TriangleVertexNumber = 0; TriangleVertexNumber < 3; ++TriangleVertexNumber )
-					{
-						FVertexAndAttributes& VertexAndAttributes = *new( PolygonToCreate.PerimeterVertices ) FVertexAndAttributes();
-						switch( TriangleVertexNumber )
-						{
-							case 0:
-								VertexAndAttributes.VertexID = bIsEdgeWindingReversedForPolygon ? EdgeVertexID1 : EdgeVertexID0;
-								break;
-
-							case 1:
-								VertexAndAttributes.VertexID = PolygonCenterVertexID;
-								break;
-
-							case 2:
-								VertexAndAttributes.VertexID = bIsEdgeWindingReversedForPolygon ? EdgeVertexID0 : EdgeVertexID1;
-								break;
-
-							default:
-								check( 0 );
-						}
+					// Set first two vertices of the new polygon directly to the original vertex instances
+					PolygonToCreate.PerimeterVertices[ 0 ].VertexInstanceID = VertexInstanceID0;
+					PolygonToCreate.PerimeterVertices[ 1 ].VertexInstanceID = VertexInstanceID1;
 
-						if( VertexAndAttributes.VertexID == PolygonCenterVertexID )
-						{
-							// Generate interpolated UVs and vertex colors for the new vertex in the center
-							{
-								const FVector CenterVertexPosition = GetVertexAttribute( PolygonCenterVertexID, UEditableMeshAttribute::VertexPosition(), 0 );
+					// Last vertex is the center vertex, so assign the new vertex
+					PolygonToCreate.PerimeterVertices[ 2 ].VertexID = PolygonCenterVertexID;
 
-								FMeshTriangle Triangle;
-								FVector TriangleVertexWeights;
-								if( ComputeBarycentricWeightForPointOnPolygon( PolygonID, CenterVertexPosition, /* Out */ Triangle, /* Out */ TriangleVertexWeights ) )
-								{
-									for( int32 TextureCoordinateIndex = 0; TextureCoordinateIndex < TextureCoordinateCount; ++TextureCoordinateIndex )
-									{
-										VertexAndAttributes.PolygonVertexAttributes.Attributes.Add( FMeshElementAttributeData(
-											UEditableMeshAttribute::VertexTextureCoordinate(),
-											TextureCoordinateIndex,
-											TriangleVertexWeights.X * GetVertexInstanceAttribute( Triangle.VertexInstanceID0, UEditableMeshAttribute::VertexTextureCoordinate(), TextureCoordinateIndex ) +
-											TriangleVertexWeights.Y * GetVertexInstanceAttribute( Triangle.VertexInstanceID1, UEditableMeshAttribute::VertexTextureCoordinate(), TextureCoordinateIndex ) +
-											TriangleVertexWeights.Z * GetVertexInstanceAttribute( Triangle.VertexInstanceID2, UEditableMeshAttribute::VertexTextureCoordinate(), TextureCoordinateIndex ) ) );
-									}
+					// Generate interpolated vertex instance attributes for the new vertex in the center
+					{
+						const FVector CenterVertexPosition = VertexPositions[ PolygonCenterVertexID ];
 
-									VertexAndAttributes.PolygonVertexAttributes.Attributes.Add( FMeshElementAttributeData(
-										UEditableMeshAttribute::VertexColor(),
-										0,
-										TriangleVertexWeights.X * GetVertexInstanceAttribute( Triangle.VertexInstanceID0, UEditableMeshAttribute::VertexColor(), 0 ) +
-										TriangleVertexWeights.Y * GetVertexInstanceAttribute( Triangle.VertexInstanceID1, UEditableMeshAttribute::VertexColor(), 0 ) +
-										TriangleVertexWeights.Z * GetVertexInstanceAttribute( Triangle.VertexInstanceID2, UEditableMeshAttribute::VertexColor(), 0 ) ) );
-								}
-							}
-						}
-						else
+						FMeshTriangle Triangle;
+						FVector TriangleVertexWeights;
+						if( ComputeBarycentricWeightForPointOnPolygon( PolygonID, CenterVertexPosition, /* Out */ Triangle, /* Out */ TriangleVertexWeights ) )
 						{
-							// Transfer over UVs and vertex colors from the original polygon that we're replacing
-							{
-								const int32 VertexNumber = FindPolygonPerimeterVertexNumberForVertex( PolygonID, VertexAndAttributes.VertexID );
-								check( VertexNumber != INDEX_NONE );
-
-								for( int32 TextureCoordinateIndex = 0; TextureCoordinateIndex < TextureCoordinateCount; ++TextureCoordinateIndex )
-								{
-									VertexAndAttributes.PolygonVertexAttributes.Attributes.Add( FMeshElementAttributeData(
-										UEditableMeshAttribute::VertexTextureCoordinate(),
-										TextureCoordinateIndex,
-										GetPolygonPerimeterVertexAttribute( PolygonID, VertexNumber, UEditableMeshAttribute::VertexTextureCoordinate(), TextureCoordinateIndex ) ) );
-								}
-
-								VertexAndAttributes.PolygonVertexAttributes.Attributes.Add( FMeshElementAttributeData(
-									UEditableMeshAttribute::VertexColor(),
-									0,
-									GetPolygonPerimeterVertexAttribute( PolygonID, VertexNumber, UEditableMeshAttribute::VertexColor(), 0 ) ) );
-							}
+							InterpAllAttributes( PolygonToCreate.PerimeterVertices[ 2 ].PolygonVertexAttributes, GetMeshDescription()->VertexInstanceAttributes(),
+								Triangle.VertexInstanceID0, Triangle.VertexInstanceID1, Triangle.VertexInstanceID2, TriangleVertexWeights );
 						}
 					}
 				}
@@ -8323,9 +7711,8 @@ void UEditableMesh::TessellatePolygons( const TArray<FPolygonID>& PolygonIDs, co
 				// the center, connecting the three new vertices that we created between each original edge.  The
 				// other three triangles will go in the corners of the original triangle.
 
-				static TArray<FVertexID> PerimeterVertexIDs;
-				GetPolygonPerimeterVertices( PolygonID, PerimeterVertexIDs );
-				check( PerimeterVertexIDs.Num() == 6 );	// We split the triangle's 3 edges earlier, so we must have six edges now
+				const TArray<FVertexInstanceID>& PerimeterVertexInstanceIDs = GetMeshDescription()->GetPolygonPerimeterVertexInstances( PolygonID );
+				check( PerimeterVertexInstanceIDs.Num() == 6 );	// We split the triangle's 3 edges earlier, so we must have six edges now
 
 				// Define the new center triangle
 				{
@@ -8334,26 +7721,10 @@ void UEditableMesh::TessellatePolygons( const TArray<FPolygonID>& PolygonIDs, co
 
 					for( int32 OriginalVertexNumber = 0; OriginalVertexNumber < 3; ++OriginalVertexNumber )
 					{
-						const int32 VertexNumber = ( OriginalVertexNumber * 2 + 1 ) % PerimeterVertexIDs.Num();
+						const int32 VertexNumber = ( OriginalVertexNumber * 2 + 1 ) % PerimeterVertexInstanceIDs.Num();
 
 						FVertexAndAttributes& VertexAndAttributes = *new( PolygonToCreate.PerimeterVertices ) FVertexAndAttributes();
-						VertexAndAttributes.VertexID = PerimeterVertexIDs[ VertexNumber ];
-
-						// Transfer over UVs and vertex colors from the original polygon that we're replacing
-						{
-							for( int32 TextureCoordinateIndex = 0; TextureCoordinateIndex < TextureCoordinateCount; ++TextureCoordinateIndex )
-							{
-								VertexAndAttributes.PolygonVertexAttributes.Attributes.Add( FMeshElementAttributeData(
-									UEditableMeshAttribute::VertexTextureCoordinate(),
-									TextureCoordinateIndex,
-									GetPolygonPerimeterVertexAttribute( PolygonID, VertexNumber, UEditableMeshAttribute::VertexTextureCoordinate(), TextureCoordinateIndex ) ) );
-							}
-
-							VertexAndAttributes.PolygonVertexAttributes.Attributes.Add( FMeshElementAttributeData(
-								UEditableMeshAttribute::VertexColor(),
-								0,
-								GetPolygonPerimeterVertexAttribute( PolygonID, VertexNumber, UEditableMeshAttribute::VertexColor(), 0 ) ) );
-						}
+						VertexAndAttributes.VertexInstanceID = PerimeterVertexInstanceIDs[ VertexNumber ];
 					}
 				}
 
@@ -8361,54 +7732,16 @@ void UEditableMesh::TessellatePolygons( const TArray<FPolygonID>& PolygonIDs, co
 				for( int32 OriginalEdgeNumber = 0; OriginalEdgeNumber < 3; ++OriginalEdgeNumber )
 				{
 					const int32 CurrentVertexNumber = OriginalEdgeNumber * 2;
-					const int32 PreviousVertexNumber = ( ( CurrentVertexNumber - 1 ) + PerimeterVertexIDs.Num() ) % PerimeterVertexIDs.Num();
-					const int32 NextVertexNumber = ( CurrentVertexNumber + 1 ) % PerimeterVertexIDs.Num();
+					const int32 PreviousVertexNumber = ( ( CurrentVertexNumber - 1 ) + PerimeterVertexInstanceIDs.Num() ) % PerimeterVertexInstanceIDs.Num();
+					const int32 NextVertexNumber = ( CurrentVertexNumber + 1 ) % PerimeterVertexInstanceIDs.Num();
 
 					FPolygonToCreate& PolygonToCreate = *new( PolygonsToCreate ) FPolygonToCreate();
 					PolygonToCreate.PolygonGroupID = PolygonGroupID;
+					PolygonToCreate.PerimeterVertices.SetNum( 3 );
 
-					for( int32 TriangleVertexNumber = 0; TriangleVertexNumber < 3; ++TriangleVertexNumber )
-					{
-						int32 VertexNumber = INDEX_NONE;
-						switch( TriangleVertexNumber )
-						{
-							case 0:
-								VertexNumber = PreviousVertexNumber;
-								break;
-
-							case 1:
-								VertexNumber = CurrentVertexNumber;
-								break;
-
-							case 2:
-								VertexNumber = NextVertexNumber;
-								break;
-
-							default:
-								check( 0 );
-						}
-
-						FVertexAndAttributes& VertexAndAttributes = *new( PolygonToCreate.PerimeterVertices ) FVertexAndAttributes();
-						VertexAndAttributes.VertexID = PerimeterVertexIDs[ VertexNumber ];
-
-						{
-							// Transfer over UVs and vertex colors from the original polygon that we're replacing
-							{
-								for( int32 TextureCoordinateIndex = 0; TextureCoordinateIndex < TextureCoordinateCount; ++TextureCoordinateIndex )
-								{
-									VertexAndAttributes.PolygonVertexAttributes.Attributes.Add( FMeshElementAttributeData(
-										UEditableMeshAttribute::VertexTextureCoordinate(),
-										TextureCoordinateIndex,
-										GetPolygonPerimeterVertexAttribute( PolygonID, VertexNumber, UEditableMeshAttribute::VertexTextureCoordinate(), TextureCoordinateIndex ) ) );
-								}
-
-								VertexAndAttributes.PolygonVertexAttributes.Attributes.Add( FMeshElementAttributeData(
-									UEditableMeshAttribute::VertexColor(),
-									0,
-									GetPolygonPerimeterVertexAttribute( PolygonID, VertexNumber, UEditableMeshAttribute::VertexColor(), 0 ) ) );
-							}
-						}
-					}
+					PolygonToCreate.PerimeterVertices[ 0 ].VertexInstanceID = PerimeterVertexInstanceIDs[ PreviousVertexNumber ];
+					PolygonToCreate.PerimeterVertices[ 1 ].VertexInstanceID = PerimeterVertexInstanceIDs[ CurrentVertexNumber ];
+					PolygonToCreate.PerimeterVertices[ 2 ].VertexInstanceID = PerimeterVertexInstanceIDs[ NextVertexNumber ];
 				}
 			}
 		}
@@ -8459,8 +7792,6 @@ void UEditableMesh::QuadrangulatePolygonGroup( const FPolygonGroupID PolygonGrou
 	// Tweakable parameters affecting how quadrangulate works
 	const float CosAngleThreshold = 0.984f;		// about 10 degrees
 	const bool bKeepHardEdges = true;
-	const bool bKeepTextureBorder = true;
-	const bool bKeepColorBorder = true;
 
 	OutNewPolygonIDs.Reset();
 
@@ -8600,6 +7931,9 @@ void UEditableMesh::QuadrangulatePolygonGroup( const FPolygonGroupID PolygonGrou
 	{
 		float BestScore = TNumericLimits<float>::Max();
 
+		const TEdgeAttributeArray<bool>& EdgeHardnesses = GetMeshDescription()->EdgeAttributes().GetAttributes<bool>( MeshAttribute::Edge::IsHard );
+		const TVertexAttributeArray<FVector>& VertexPositions = GetMeshDescription()->VertexAttributes().GetAttributes<FVector>( MeshAttribute::Vertex::Position );
+
 		for( const FPolygonID PolygonID : PolygonIDs )
 		{
 			// If it's not a triangle, don't consider this polygon at all
@@ -8617,7 +7951,7 @@ void UEditableMesh::QuadrangulatePolygonGroup( const FPolygonGroupID PolygonGrou
 				bool bOutEdgeWindingIsReversedForPolygon;
 				const FEdgeID PerimeterEdgeID = GetPolygonPerimeterEdge( PolygonID, PerimeterEdgeIndex, bOutEdgeWindingIsReversedForPolygon );
 
-				const bool bIsSoftEdge = FMath::IsNearlyZero( GetEdgeAttribute( PerimeterEdgeID, UEditableMeshAttribute::EdgeIsHard(), 0 ).X );
+				const bool bIsSoftEdge = !EdgeHardnesses[ PerimeterEdgeID ];
 				if( !bKeepHardEdges || bIsSoftEdge )
 				{
 					const FPolygonID AdjacentPolygonID = [ this, PerimeterEdgeID, PolygonID ]() -> FPolygonID
@@ -8679,40 +8013,21 @@ void UEditableMesh::QuadrangulatePolygonGroup( const FPolygonGroupID PolygonGrou
 								const int32 PrevAdjacentPerimeterEdgeIndex = ( AdjacentPerimeterEdgeIndex + 2 ) % 3;
 								const int32 NextAdjacentPerimeterEdgeIndex = ( AdjacentPerimeterEdgeIndex + 1 ) % 3;
 
-								// If we wish to maintain texture seams, compare the vertex UVs of each polygon's vertex along the edge to be removed, and only proceed if they are equal
-								bool bTextureCoordinatesEqual = true;
-								if( bKeepTextureBorder )
-								{
-									for( int32 TextureCoordinateIndex = 0; TextureCoordinateIndex < TextureCoordinateCount; ++TextureCoordinateIndex )
-									{
-										const FVector4 UVPolygon1Vertex1 = GetPolygonPerimeterVertexAttribute( PolygonID, PerimeterEdgeIndex, UEditableMeshAttribute::VertexTextureCoordinate(), TextureCoordinateIndex );
-										const FVector4 UVPolygon2Vertex1 = GetPolygonPerimeterVertexAttribute( AdjacentPolygonID, AdjacentPerimeterEdgeIndex, UEditableMeshAttribute::VertexTextureCoordinate(), TextureCoordinateIndex );
-										const FVector4 UVPolygon1Vertex2 = GetPolygonPerimeterVertexAttribute( PolygonID, NextPerimeterEdgeIndex, UEditableMeshAttribute::VertexTextureCoordinate(), TextureCoordinateIndex );
-										const FVector4 UVPolygon2Vertex2 = GetPolygonPerimeterVertexAttribute( AdjacentPolygonID, PrevAdjacentPerimeterEdgeIndex, UEditableMeshAttribute::VertexTextureCoordinate(), TextureCoordinateIndex );
-										if( UVPolygon1Vertex1 != UVPolygon2Vertex1 || UVPolygon1Vertex2 != UVPolygon2Vertex2 )
-										{
-											bTextureCoordinatesEqual = false;
-											break;
-										}
-									}
-								}
+								const TArray<FVertexInstanceID>& PerimeterVertexInstanceIDs = GetMeshDescription()->GetPolygonPerimeterVertexInstances( PolygonID );
+								const TArray<FVertexInstanceID>& AdjacentPerimeterVertexInstanceIDs = GetMeshDescription()->GetPolygonPerimeterVertexInstances( AdjacentPolygonID );
 
-								// If we wish to maintain color seams, compare the vertex colors of each polygon's vertex along the edge to be removed, and only proceed if they are equal
-								bool bColorsEqual = true;
-								if( bKeepColorBorder )
-								{
-									const FVector4 ColorPolygon1Vertex1 = GetPolygonPerimeterVertexAttribute( PolygonID, PerimeterEdgeIndex, UEditableMeshAttribute::VertexColor(), 0 );
-									const FVector4 ColorPolygon2Vertex1 = GetPolygonPerimeterVertexAttribute( AdjacentPolygonID, AdjacentPerimeterEdgeIndex, UEditableMeshAttribute::VertexColor(), 0 );
-									const FVector4 ColorPolygon1Vertex2 = GetPolygonPerimeterVertexAttribute( PolygonID, NextPerimeterEdgeIndex, UEditableMeshAttribute::VertexColor(), 0 );
-									const FVector4 ColorPolygon2Vertex2 = GetPolygonPerimeterVertexAttribute( AdjacentPolygonID, PrevAdjacentPerimeterEdgeIndex, UEditableMeshAttribute::VertexColor(), 0 );
-									if( ColorPolygon1Vertex1 != ColorPolygon2Vertex1 || ColorPolygon1Vertex2 != ColorPolygon2Vertex2 )
-									{
-										bColorsEqual = false;
-									}
-								}
+								const FVertexInstanceID VertexInstanceStart1 = PerimeterVertexInstanceIDs[ PerimeterEdgeIndex ];
+								const FVertexInstanceID VertexInstanceStart2 = AdjacentPerimeterVertexInstanceIDs[ AdjacentPerimeterEdgeIndex ];
+								const FVertexInstanceID VertexInstanceEnd1 = PerimeterVertexInstanceIDs[ NextPerimeterEdgeIndex ];
+								const FVertexInstanceID VertexInstanceEnd2 = AdjacentPerimeterVertexInstanceIDs[ PrevAdjacentPerimeterEdgeIndex ];
+
+								auto AttributePred = []( const auto AttributeName, const auto& AttributeIndicesArray ) { return EnumHasAllFlags( AttributeIndicesArray.GetFlags(), EMeshAttributeFlags::Mergeable ); };
 
-								// If we are maintaining texture or color seams, only merge the polygons if the appropriate attributes are equal for each polygon
-								if( ( !bKeepTextureBorder || bTextureCoordinatesEqual ) && ( !bKeepColorBorder || bColorsEqual ) )
+								const bool bAttributesEqual = ( VertexInstanceStart1 == VertexInstanceStart2 && VertexInstanceEnd1 == VertexInstanceEnd2 ) ||
+									( CompareAttributesByPredicate( GetMeshDescription()->VertexInstanceAttributes(), VertexInstanceStart1, VertexInstanceStart2, AttributePred ) &&
+									  CompareAttributesByPredicate( GetMeshDescription()->VertexInstanceAttributes(), VertexInstanceEnd1, VertexInstanceEnd2, AttributePred ) );
+
+								if( bAttributesEqual )
 								{
 									const FVertexID V0 = GetPolygonPerimeterVertex( PolygonID, PrevPerimeterEdgeIndex );
 									const FVertexID V1 = GetPolygonPerimeterVertex( PolygonID, PerimeterEdgeIndex );
@@ -8720,10 +8035,10 @@ void UEditableMesh::QuadrangulatePolygonGroup( const FPolygonGroupID PolygonGrou
 									const FVertexID V3 = GetPolygonPerimeterVertex( AdjacentPolygonID, PrevAdjacentPerimeterEdgeIndex );
 									check( V3 == GetPolygonPerimeterVertex( PolygonID, NextPerimeterEdgeIndex ) );
 
-									const FVector P0 = GetVertexAttribute( V0, UEditableMeshAttribute::VertexPosition(), 0 );
-									const FVector P1 = GetVertexAttribute( V1, UEditableMeshAttribute::VertexPosition(), 0 );
-									const FVector P2 = GetVertexAttribute( V2, UEditableMeshAttribute::VertexPosition(), 0 );
-									const FVector P3 = GetVertexAttribute( V3, UEditableMeshAttribute::VertexPosition(), 0 );
+									const FVector P0 = VertexPositions[ V0 ];
+									const FVector P1 = VertexPositions[ V1 ];
+									const FVector P2 = VertexPositions[ V2 ];
+									const FVector P3 = VertexPositions[ V3 ];
 
 									const FVector D01 = ( P1 - P0 ).GetSafeNormal();
 									const FVector D12 = ( P2 - P1 ).GetSafeNormal();
@@ -8816,34 +8131,10 @@ void UEditableMesh::QuadrangulatePolygonGroup( const FPolygonGroupID PolygonGrou
 			PolygonToCreate.PerimeterVertices.Emplace();
 			FVertexAndAttributes& VertexAndAttributes = PolygonToCreate.PerimeterVertices.Last();
 
-			VertexAndAttributes.VertexID = GetPolygonPerimeterVertex(
+			// Build polygon from original vertex instances, to save creating new ones and copying the original attributes over
+			VertexAndAttributes.VertexInstanceID = GetPolygonPerimeterVertexInstance(
 				AdjacentPolygon1.Vertices[ VertexIndex ].Get<0>(),
 				AdjacentPolygon1.Vertices[ VertexIndex ].Get<1>() );
-
-			TArray<FMeshElementAttributeData>& Attributes = VertexAndAttributes.PolygonVertexAttributes.Attributes;
-
-			for( int32 TextureCoordinateIndex = 0; TextureCoordinateIndex < TextureCoordinateCount; ++TextureCoordinateIndex )
-			{
-				Attributes.Emplace(
-					UEditableMeshAttribute::VertexTextureCoordinate(),
-					TextureCoordinateIndex,
-					GetPolygonPerimeterVertexAttribute( 
-						AdjacentPolygon1.Vertices[ VertexIndex ].Get<0>(),
-						AdjacentPolygon1.Vertices[ VertexIndex ].Get<1>(),
-						UEditableMeshAttribute::VertexTextureCoordinate(),
-						TextureCoordinateIndex ) 
-					);
-			}
-
-			Attributes.Emplace(
-				UEditableMeshAttribute::VertexColor(),
-				0,
-				GetPolygonPerimeterVertexAttribute(
-					AdjacentPolygon1.Vertices[ VertexIndex ].Get<0>(),
-					AdjacentPolygon1.Vertices[ VertexIndex ].Get<1>(),
-					UEditableMeshAttribute::VertexColor(),
-					0 )
-				);
 		}
 
 		// Specify old polygons to be deleted
@@ -8974,7 +8265,7 @@ void UEditableMesh::QuadrangulatePolygonGroup( const FPolygonGroupID PolygonGrou
 
 	const bool bDeleteOrphanedEdges = true;
 	const bool bDeleteOrphanedVertices = false;
-	const bool bDeleteOrphanedVertexInstances = false;
+	const bool bDeleteOrphanedVertexInstances = true;
 	const bool bDeleteEmptyPolygonGroups = false;
 	DeletePolygons( PolygonIDsToDelete, bDeleteOrphanedEdges, bDeleteOrphanedVertices, bDeleteOrphanedVertexInstances, bDeleteEmptyPolygonGroups );
 }
@@ -9029,6 +8320,7 @@ void UEditableMesh::UpdateOrRebuildOctree()
 
 			// Now, add new polygons to the octree
 			{
+				const TVertexAttributeArray<FVector>& VertexPositions = GetMeshDescription()->VertexAttributes().GetAttributes<FVector>( MeshAttribute::Vertex::Position );
 				for( const FPolygonID PolygonID : NewOctreePolygonIDs )
 				{
 					checkSlow( IsValidPolygon( PolygonID ) );
@@ -9041,8 +8333,7 @@ void UEditableMesh::UpdateOrRebuildOctree()
 						GetPolygonPerimeterVertices( PolygonID, /* Out */ PerimeterVertices );
 						for( const FVertexID VertexID : PerimeterVertices )
 						{
-							const FVector VertexPosition = GetVertexAttribute( VertexID, UEditableMeshAttribute::VertexPosition(), 0 );
-							BoundingBox += VertexPosition;
+							BoundingBox += VertexPositions[ VertexID ];
 						}
 					}
 
@@ -9098,6 +8389,7 @@ void UEditableMesh::RebuildOctree()
 		// @todo mesheditor spatial: We'll currently only ever find vertices and edges that are connected to polygons.  Our spatial
 		// database only tracks polygons.  Do we need to support hover/selection of "loose" vertices and edges also?
 
+		const TVertexAttributeArray<FVector>& VertexPositions = GetMeshDescription()->VertexAttributes().GetAttributes<FVector>( MeshAttribute::Vertex::Position );
 		for( const FPolygonID PolygonID : GetMeshDescription()->Polygons().GetElementIDs() )
 		{
 			FBox BoundingBox;
@@ -9108,8 +8400,7 @@ void UEditableMesh::RebuildOctree()
 				GetPolygonPerimeterVertices( PolygonID, /* Out */ PerimeterVertices );
 				for( const FVertexID VertexID : PerimeterVertices )
 				{
-					const FVector VertexPosition = GetVertexAttribute( VertexID, UEditableMeshAttribute::VertexPosition(), 0 );
-					BoundingBox += VertexPosition;
+					BoundingBox += VertexPositions[ VertexID ];
 				}
 			}
 
