Total revamp of mesh element attribute model.

diff --git a/Engine/Plugins/Runtime/EditableMesh/Source/EditableMesh/EditableStaticMeshAdapter.cpp b/Engine/Plugins/Runtime/EditableMesh/Source/EditableMesh/EditableStaticMeshAdapter.cpp
--- a/Engine/Plugins/Runtime/EditableMesh/Source/EditableMesh/EditableStaticMeshAdapter.cpp
+++ b/Engine/Plugins/Runtime/EditableMesh/Source/EditableMesh/EditableStaticMeshAdapter.cpp
@@ -9,6 +9,7 @@
 #include "PhysicsEngine/BodySetup.h"	// For collision generation
 #include "ProfilingDebugging/ScopedTimers.h"	// For FAutoScopedDurationTimer
 #include "EditableMeshFactory.h"
+#include "MeshAttributes.h"
 
 
 const FTriangleID FTriangleID::Invalid( TNumericLimits<uint32>::Max() );
@@ -91,17 +92,36 @@ void UEditableStaticMeshAdapter::InitEditableStaticMesh( UEditableMesh* Editable
 
 				const FStaticMeshLODResources& StaticMeshLOD = StaticMeshRenderData.LODResources[ StaticMeshLODIndex ];
 
+				UMeshDescription* MeshDescription = EditableMesh->GetMeshDescription();
+
+				TVertexAttributeArray<FVector>& VertexPositions = MeshDescription->VertexAttributes().GetAttributes<FVector>( MeshAttribute::Vertex::Position );
+
+				TVertexInstanceAttributeArray<FVector>& VertexInstanceNormals = MeshDescription->VertexInstanceAttributes().GetAttributes<FVector>( MeshAttribute::VertexInstance::Normal );
+				TVertexInstanceAttributeArray<FVector>& VertexInstanceTangents = MeshDescription->VertexInstanceAttributes().GetAttributes<FVector>( MeshAttribute::VertexInstance::Tangent );
+				TVertexInstanceAttributeArray<float>& VertexInstanceBinormalSigns = MeshDescription->VertexInstanceAttributes().GetAttributes<float>( MeshAttribute::VertexInstance::BinormalSign );
+				TVertexInstanceAttributeArray<FVector4>& VertexInstanceColors = MeshDescription->VertexInstanceAttributes().GetAttributes<FVector4>( MeshAttribute::VertexInstance::Color );
+				TVertexInstanceAttributeIndicesArray<FVector2D>& VertexInstanceUVs = MeshDescription->VertexInstanceAttributes().GetAttributesSet<FVector2D>( MeshAttribute::VertexInstance::TextureCoordinate );
+
+				TEdgeAttributeArray<bool>& EdgeHardnesses = MeshDescription->EdgeAttributes().GetAttributes<bool>( MeshAttribute::Edge::IsHard );
+
+				TPolygonGroupAttributeArray<UObject*>& PolygonGroupMaterialAssets = MeshDescription->PolygonGroupAttributes().GetAttributes<UObject*>( MeshAttribute::PolygonGroup::MaterialAsset );
+				TPolygonGroupAttributeArray<FName>& PolygonGroupMaterialSlotNames = MeshDescription->PolygonGroupAttributes().GetAttributes<FName>( MeshAttribute::PolygonGroup::MaterialSlotName );
+				TPolygonGroupAttributeArray<FName>& PolygonGroupImportedMaterialSlotNames = MeshDescription->PolygonGroupAttributes().GetAttributes<FName>( MeshAttribute::PolygonGroup::ImportedMaterialSlotName );
+				TPolygonGroupAttributeArray<bool>& PolygonGroupCollision = MeshDescription->PolygonGroupAttributes().GetAttributes<bool>( MeshAttribute::PolygonGroup::EnableCollision );
+				TPolygonGroupAttributeArray<bool>& PolygonGroupCastShadow = MeshDescription->PolygonGroupAttributes().GetAttributes<bool>( MeshAttribute::PolygonGroup::CastShadow );
+
 				// Store off the number of texture coordinates in this mesh
-				EditableMesh->TextureCoordinateCount = StaticMeshLOD.GetNumTexCoords();
+				const int32 NumUVs = StaticMeshLOD.GetNumTexCoords();
+				VertexInstanceUVs.SetNumIndices(NumUVs);
+				EditableMesh->TextureCoordinateCount = NumUVs;
 
 				// Vertices
 				const int32 NumRenderingVertices = StaticMeshLOD.PositionVertexBuffer.GetNumVertices();
-				const int32 NumUVs = StaticMeshLOD.GetNumTexCoords();
 				const bool bHasColor = StaticMeshLOD.ColorVertexBuffer.GetNumVertices() > 0;
 				check( !bHasColor || StaticMeshLOD.ColorVertexBuffer.GetNumVertices() == StaticMeshLOD.VertexBuffer.GetNumVertices() );
 
-				FVertexArray& Vertices = EditableMesh->GetMeshDescription()->Vertices();
-				FVertexInstanceArray& VertexInstances = EditableMesh->GetMeshDescription()->VertexInstances();
+				FVertexArray& Vertices = MeshDescription->Vertices();
+				FVertexInstanceArray& VertexInstances = MeshDescription->VertexInstances();
 
 				// @todo mesheditor cleanup: This code is very similar to the static mesh build code; try to share helper structs
 				static TMultiMap<int32, int32> OverlappingRenderingVertexIndices;
@@ -200,14 +220,10 @@ void UEditableStaticMeshAdapter::InitEditableStaticMesh( UEditableMesh* Editable
 							{
 								check( VertexInstances.IsValid( OverlappingVertexInstanceID ) );
 								const FVertexID ExistingVertexID = VertexInstances[ OverlappingVertexInstanceID ].VertexID;
-								FMeshVertex& ExistingVertex = Vertices[ ExistingVertexID ];
 
 								// We already have a unique editable vertex for this vertex instance position, so link them!
-								VertexInstances.Insert( VertexInstanceID );
-								VertexInstances[ VertexInstanceID ].VertexID = ExistingVertexID;
+								MeshDescription->CreateVertexInstanceWithID( VertexInstanceID, ExistingVertexID );
 
-								checkSlow( !ExistingVertex.VertexInstanceIDs.Contains( VertexInstanceID ) );
-								ExistingVertex.VertexInstanceIDs.Add( VertexInstanceID );
 								bAlreadyHaveVertexForPosition = true;
 
 								break;
@@ -217,41 +233,28 @@ void UEditableStaticMeshAdapter::InitEditableStaticMesh( UEditableMesh* Editable
 
 					if( !bAlreadyHaveVertexForPosition )
 					{
-						const FVertexID NewVertexID = Vertices.Add();
-						FMeshVertex& NewVertex = Vertices[ NewVertexID ];
-						NewVertex.VertexPosition = VertexPosition;
-						NewVertex.CornerSharpness = 0.0f;
-
-						VertexInstances.Insert( VertexInstanceID );
-						VertexInstances[ VertexInstanceID ].VertexID = NewVertexID;
-						
-						// @todo mesheditor: If a mesh somehow contained vertex instances that no triangle was referencing, this would cause
-						// the vertex instance to be ignored by the editable mesh code.  It would just sit in the vertex buffer (and in the
-						// editable mesh vertex's RenderingVertexIndices list), but would never be touched.  The editable mesh code only
-						// creates vertex instances for vertices that are attached to polygons, so this should never happen with meshes
-						// that we create and save.  Only if the incoming data had orphan vertices in it.  Should hopefully not be a problem.
-						NewVertex.VertexInstanceIDs.Add( VertexInstanceID );
-
-						// NOTE: The new vertex's connected polygons will be filled in down below, as we're processing mesh triangles
+						const FVertexID NewVertexID = MeshDescription->CreateVertex();
+						VertexPositions[ NewVertexID ] = VertexPosition;
+
+						MeshDescription->CreateVertexInstanceWithID( VertexInstanceID, NewVertexID );
 					}
 
 					// Populate the vertex instance attributes
 					{
-						FMeshVertexInstance& VertexInstance = VertexInstances[ VertexInstanceID ];
-						VertexInstance.VertexUVs.Reserve( NumUVs );
-						for( int32 UVIndex = 0; UVIndex < NumUVs; ++UVIndex )
-						{
-							VertexInstance.VertexUVs.Add( StaticMeshLOD.VertexBuffer.GetVertexUV( RenderingVertexIndex, UVIndex ) );
-						}
-
 						const FVector Normal = StaticMeshLOD.VertexBuffer.VertexTangentZ( RenderingVertexIndex );
 						const FVector Tangent = StaticMeshLOD.VertexBuffer.VertexTangentX( RenderingVertexIndex );
 						const FVector Binormal = StaticMeshLOD.VertexBuffer.VertexTangentY( RenderingVertexIndex );
-						VertexInstance.Normal = Normal;
-						VertexInstance.Tangent = Tangent;
-						VertexInstance.BinormalSign = GetBasisDeterminantSign(Tangent, Binormal, Normal);
+						const float BinormalSign = GetBasisDeterminantSign(Tangent, Binormal, Normal);
+						const FLinearColor Color = bHasColor ? FLinearColor( StaticMeshLOD.ColorVertexBuffer.VertexColor( RenderingVertexIndex ) ) : FLinearColor::White;
 
-						VertexInstance.Color = bHasColor ? FLinearColor( StaticMeshLOD.ColorVertexBuffer.VertexColor( RenderingVertexIndex ) ) : FLinearColor::White;
+						VertexInstanceNormals[ VertexInstanceID ] = Normal;
+						VertexInstanceTangents[ VertexInstanceID ] = Tangent;
+						VertexInstanceBinormalSigns[ VertexInstanceID ] = BinormalSign;
+						VertexInstanceColors[ VertexInstanceID ] = Color;
+						for( int32 UVIndex = 0; UVIndex < NumUVs; ++UVIndex )
+						{
+							VertexInstanceUVs.GetArrayForIndex( UVIndex )[ VertexInstanceID ] = StaticMeshLOD.VertexBuffer.GetVertexUV( RenderingVertexIndex, UVIndex );
+						}
 					}
 
 				}
@@ -259,10 +262,6 @@ void UEditableStaticMeshAdapter::InitEditableStaticMesh( UEditableMesh* Editable
 
 				const FIndexArrayView RenderingIndices = StaticMeshLOD.IndexBuffer.GetArrayView();
 
-
-				static TMap<uint64, FEdgeID> UniqueEdgeToEdgeID;
-				UniqueEdgeToEdgeID.Reset();
-
 				static TMultiMap<FEdgeID, TTuple<FVertexInstanceID, FVertexInstanceID>> EdgeToVertexInstancePair;
 				EdgeToVertexInstancePair.Reset();
 
@@ -275,17 +274,16 @@ void UEditableStaticMeshAdapter::InitEditableStaticMesh( UEditableMesh* Editable
 				for( uint32 RenderingSectionIndex = 0; RenderingSectionIndex < NumSections; ++RenderingSectionIndex )
 				{
 					const FStaticMeshSection& RenderingSection = StaticMeshLOD.Sections[ RenderingSectionIndex ];
+					const FStaticMaterial& StaticMaterial = StaticMesh->StaticMaterials[ RenderingSection.MaterialIndex ];
+					UMaterialInterface* MaterialInterface = StaticMaterial.MaterialInterface;
 
 					// Create a new polygon group
-					const FPolygonGroupID NewPolygonGroupID = PolygonGroups.Add();
-					FMeshPolygonGroup& NewPolygonGroup = PolygonGroups[ NewPolygonGroupID ];
-
-					const FStaticMaterial& StaticMaterial = StaticMesh->StaticMaterials[ RenderingSection.MaterialIndex ];
-					NewPolygonGroup.MaterialAsset = StaticMaterial.MaterialInterface->GetPathName();
-					NewPolygonGroup.MaterialSlotName = StaticMaterial.MaterialSlotName;
-					NewPolygonGroup.ImportedMaterialSlotName = StaticMaterial.ImportedMaterialSlotName;
-					NewPolygonGroup.bEnableCollision = RenderingSection.bEnableCollision;
-					NewPolygonGroup.bCastShadow = RenderingSection.bCastShadow;
+					const FPolygonGroupID NewPolygonGroupID = MeshDescription->CreatePolygonGroup();
+					PolygonGroupMaterialAssets[ NewPolygonGroupID ] = MaterialInterface;
+					PolygonGroupMaterialSlotNames[ NewPolygonGroupID ] = StaticMaterial.MaterialSlotName;
+					PolygonGroupImportedMaterialSlotNames[ NewPolygonGroupID ] = StaticMaterial.ImportedMaterialSlotName;
+					PolygonGroupCollision[ NewPolygonGroupID ] = RenderingSection.bEnableCollision;
+					PolygonGroupCastShadow[ NewPolygonGroupID ] = RenderingSection.bCastShadow;
 
 					// Create a rendering polygon group for holding the triangulated data and references to the static mesh rendering section.
 					// This is indexed by the same FPolygonGroupID as the PolygonGroups.
@@ -316,141 +314,91 @@ void UEditableStaticMeshAdapter::InitEditableStaticMesh( UEditableMesh* Editable
 							TriangleVertexIDs[ 0 ] != TriangleVertexIDs[ 1 ] &&
 							TriangleVertexIDs[ 1 ] != TriangleVertexIDs[ 2 ] &&
 							TriangleVertexIDs[ 2 ] != TriangleVertexIDs[ 0 ];
+
 						if( bIsValidTriangle )
 						{
-							// Static meshes only support triangles, so there's no need to triangulate anything yet.  We'll make both
-							// a triangle and a polygon here.
-							const FTriangleID NewTriangleID = FTriangleID( SectionTriangleIndex );
-
-							NewRenderingPolygonGroup.Triangles.Insert( NewTriangleID );
-							FMeshTriangle& NewTriangle = NewRenderingPolygonGroup.Triangles[ NewTriangleID ];
-
-							// Insert a polygon into the mesh
-							const FPolygonID NewPolygonID = Polygons.Add();
-							FMeshPolygon& NewPolygon = Polygons[ NewPolygonID ];
-							NewPolygon.PolygonGroupID = NewPolygonGroupID;
-							NewPolygonGroup.Polygons.Add( NewPolygonID );
-
-							// Create a rendering polygon mirror, indexed by the same ID
-							RenderingPolygons.Insert( NewPolygonID );
-							FRenderingPolygon& NewRenderingPolygon = RenderingPolygons[ NewPolygonID ];
-							NewRenderingPolygon.PolygonGroupID = NewPolygonGroupID;
-							NewRenderingPolygon.TriangulatedPolygonTriangleIndices.Add( NewTriangleID );
-
-
-							// Static meshes don't support polygons with holes, so we always start out with only a perimeter contour per polygon
-							FMeshPolygonContour& PerimeterContour = NewPolygon.PerimeterContour;
-							PerimeterContour.VertexInstanceIDs.Reserve( 3 );
-
-							// Connect vertices
-							for( uint32 TriangleVertexIndex = 0; TriangleVertexIndex < 3; ++TriangleVertexIndex )
-							{
-								const FVertexID VertexID = TriangleVertexIDs[ TriangleVertexIndex ];
-								const FVertexInstanceID VertexInstanceID = TriangleVertexInstanceIDs[ TriangleVertexIndex ];
-
-								NewPolygon.PerimeterContour.VertexInstanceIDs.Add( VertexInstanceID );
-
-								FMeshVertexInstance& VertexInstance = VertexInstances[ VertexInstanceID ];
-								VertexInstance.ConnectedPolygons.Add( NewPolygonID );
-
-								// The triangle points to each of its three vertices
-								NewTriangle.SetVertexInstanceID( TriangleVertexIndex, VertexInstanceID );
-							}
-
-							// Add triangle to polygon triangulation array
-							NewPolygon.Triangles.Add( NewTriangle );
+							FEdgeID NewEdgeIDs[ 3 ];
 
 							// Connect edges
 							{
-								struct Local
-								{
-									inline static uint64 Make64BitValueForEdge( const FVertexID EdgeVertexID0, const FVertexID EdgeVertexID1 )
-									{
-										return uint64( ( uint64( EdgeVertexID0.GetValue() ) << 32 ) | uint64( EdgeVertexID1.GetValue() ) );
-									};
-								};
-
-
 								// Add the edges of this triangle
 								for( uint32 TriangleEdgeNumber = 0; TriangleEdgeNumber < 3; ++TriangleEdgeNumber )
 								{
-									FVertexInstanceID EdgeVertexInstanceIDs[ 2 ];
-									EdgeVertexInstanceIDs[ 0 ] = TriangleVertexInstanceIDs[ TriangleEdgeNumber ];
-									EdgeVertexInstanceIDs[ 1 ] = TriangleVertexInstanceIDs[ ( TriangleEdgeNumber + 1 ) % 3 ];
+									const FVertexID VertexID0 = TriangleVertexIDs[ TriangleEdgeNumber ];
+									const FVertexID VertexID1 = TriangleVertexIDs[ ( TriangleEdgeNumber + 1 ) % 3 ];
+									
+									FEdgeID NewEdgeID = MeshDescription->GetVertexPairEdge( VertexID0, VertexID1 );
 
-									FVertexID EdgeVertexIDs[ 2 ];
-									EdgeVertexIDs[ 0 ] = TriangleVertexIDs[ TriangleEdgeNumber ];
-									EdgeVertexIDs[ 1 ] = TriangleVertexIDs[ ( TriangleEdgeNumber + 1 ) % 3 ];
-
-									// Check to see if this edge already exists
-									bool bAlreadyHaveEdge = false;
-									FEdgeID EdgeID = FEdgeID::Invalid;
+									if( NewEdgeID == FEdgeID::Invalid )
 									{
-										FEdgeID* FoundEdgeIDPtr = UniqueEdgeToEdgeID.Find( Local::Make64BitValueForEdge( EdgeVertexIDs[ 0 ], EdgeVertexIDs[ 1 ] ) );
-										if( FoundEdgeIDPtr != nullptr )
-										{
-											bAlreadyHaveEdge = true;
-											EdgeID = *FoundEdgeIDPtr;
-										}
-										else
-										{
-											// Try the other way around
-											FoundEdgeIDPtr = UniqueEdgeToEdgeID.Find( Local::Make64BitValueForEdge( EdgeVertexIDs[ 1 ], EdgeVertexIDs[ 0 ] ) );
-
-											if( FoundEdgeIDPtr != nullptr )
-											{
-												Swap( EdgeVertexIDs[ 0 ], EdgeVertexIDs[ 1 ] );
-												bAlreadyHaveEdge = true;
-												EdgeID = *FoundEdgeIDPtr;
-											}
-										}
+										NewEdgeID = MeshDescription->CreateEdge( VertexID0, VertexID1 );
 									}
 
-									if( !bAlreadyHaveEdge )
-									{
-										// Create the new edge.  We'll connect it to its polygons later on.
-										EdgeID = Edges.Add();
-										FMeshEdge& NewEdge = Edges[ EdgeID ];
-
-										NewEdge.VertexIDs[ 0 ] = EdgeVertexIDs[ 0 ];
-										NewEdge.VertexIDs[ 1 ] = EdgeVertexIDs[ 1 ];
-										NewEdge.bIsHardEdge = false;
-										NewEdge.CreaseSharpness = 0.0f;
+									NewEdgeIDs[ TriangleEdgeNumber ] = NewEdgeID;
 
-										UniqueEdgeToEdgeID.Add( Local::Make64BitValueForEdge( EdgeVertexIDs[ 0 ], EdgeVertexIDs[ 1 ] ), EdgeID );
-									}
-
-									// Each edge will point back to the polygon that its connected to.  Remember, an edge can be shared by multiple
-									// polygons, but usually its best if only shared by up to two.
-									FMeshEdge& Edge = Edges[ EdgeID ];
-
-									Edge.ConnectedPolygons.AddUnique( NewPolygonID );
-
-									// Connect the end vertices to the edge
-									Vertices[ EdgeVertexIDs[ 0 ] ].ConnectedEdgeIDs.AddUnique( EdgeID );
-									Vertices[ EdgeVertexIDs[ 1 ] ].ConnectedEdgeIDs.AddUnique( EdgeID );
-
-									// Determine whether the edge is hard by checking whether
-									if( !Edge.bIsHardEdge )
+									// Determine edge hardness by checking whether overlapping edges' vertex instances have the same normals or not
+									if( !EdgeHardnesses[ NewEdgeID ] )
 									{
+										// Get vertex instance IDs corresponding to the ends of the edge.
+										// Always order them the same for overlapping edges so we compare equivalent vertex instances each time.
+										// For simplicity, we just order them by vertex ID value, lowest/highest.
+										FVertexInstanceID VertexInstanceID0 = TriangleVertexInstanceIDs[ TriangleEdgeNumber ];
+										FVertexInstanceID VertexInstanceID1 = TriangleVertexInstanceIDs[ ( TriangleEdgeNumber + 1 ) % 3 ];
+										if( VertexID0.GetValue() > VertexID1.GetValue() )
+										{
+											Swap( VertexInstanceID0, VertexInstanceID1 );
+										}
+
 										TArray<TTuple<FVertexInstanceID, FVertexInstanceID>> VertexInstancePairs;
-										EdgeToVertexInstancePair.MultiFind( EdgeID, VertexInstancePairs );
+										EdgeToVertexInstancePair.MultiFind( NewEdgeID, VertexInstancePairs );
 
 										for( const TTuple<FVertexInstanceID, FVertexInstanceID>& VertexInstancePair : VertexInstancePairs )
 										{
-											// If either of the ends of the edge just added have different normals to any overlapping edge, mark it as a hard edge
-											if( VertexInstances[ VertexInstancePair.Get<0>() ].Normal != VertexInstances[ EdgeVertexInstanceIDs[ 0 ] ].Normal ||
-												VertexInstances[ VertexInstancePair.Get<1>() ].Normal != VertexInstances[ EdgeVertexInstanceIDs[ 1 ] ].Normal )
+											const FVertexInstanceID ExistingVertexInstanceID0 = VertexInstancePair.Get<0>();
+											const FVertexInstanceID ExistingVertexInstanceID1 = VertexInstancePair.Get<1>();
+
+											if( VertexInstanceNormals[ ExistingVertexInstanceID0 ] != VertexInstanceNormals[ VertexInstanceID0 ] ||
+												VertexInstanceNormals[ ExistingVertexInstanceID1 ] != VertexInstanceNormals[ VertexInstanceID1 ] )
 											{
-												Edge.bIsHardEdge = true;
+												EdgeHardnesses[ NewEdgeID ] = true;
 												break;
 											}
 										}
 
-										EdgeToVertexInstancePair.Add( EdgeID, MakeTuple( EdgeVertexInstanceIDs[ 0 ], EdgeVertexInstanceIDs[ 1 ] ) );
+										EdgeToVertexInstancePair.Add( NewEdgeID, MakeTuple( VertexInstanceID0, VertexInstanceID1 ) );
 									}
 								}
 							}
+
+							// Static meshes only support triangles, so there's no need to triangulate anything yet.  We'll make both
+							// a triangle and a polygon here.
+							const FTriangleID NewTriangleID = FTriangleID( SectionTriangleIndex );
+
+							NewRenderingPolygonGroup.Triangles.Insert( NewTriangleID );
+							FMeshTriangle& NewTriangle = NewRenderingPolygonGroup.Triangles[ NewTriangleID ];
+
+							static TArray<UMeshDescription::FContourPoint> Perimeter;
+							Perimeter.Reset( 3 );
+							Perimeter.AddUninitialized( 3 );
+							for( uint32 TriangleVertexIndex = 0; TriangleVertexIndex < 3; ++TriangleVertexIndex )
+							{
+								Perimeter[ TriangleVertexIndex ].VertexInstanceID = TriangleVertexInstanceIDs[ TriangleVertexIndex ];
+								Perimeter[ TriangleVertexIndex ].EdgeID = NewEdgeIDs[ TriangleVertexIndex ];
+
+								NewTriangle.SetVertexInstanceID( TriangleVertexIndex, TriangleVertexInstanceIDs[ TriangleVertexIndex ] );
+							}
+
+							// Insert a polygon into the mesh
+							const FPolygonID NewPolygonID = MeshDescription->CreatePolygon( NewPolygonGroupID, Perimeter );
+
+							// Create a rendering polygon mirror, indexed by the same ID
+							RenderingPolygons.Insert( NewPolygonID );
+							FRenderingPolygon& NewRenderingPolygon = RenderingPolygons[ NewPolygonID ];
+							NewRenderingPolygon.PolygonGroupID = NewPolygonGroupID;
+							NewRenderingPolygon.TriangulatedPolygonTriangleIndices.Add( NewTriangleID );
+
+							// Add triangle to polygon triangulation array
+							MeshDescription->GetPolygonTriangles( NewPolygonID ).Add( NewTriangle );
 						}
 						else
 						{
@@ -579,6 +527,13 @@ void UEditableStaticMeshAdapter::OnRebuildRenderMesh( const UEditableMesh* Edita
 
 	bool bHasColor = false;
 
+	const UMeshDescription* MeshDescription = EditableMesh->GetMeshDescription();
+	check( MeshDescription );
+
+	const TPolygonGroupAttributeArray<FName>& PolygonGroupImportedMaterialSlotNames = MeshDescription->PolygonGroupAttributes().GetAttributes<FName>( MeshAttribute::PolygonGroup::ImportedMaterialSlotName );
+	const TPolygonGroupAttributeArray<bool>& PolygonGroupCollision = MeshDescription->PolygonGroupAttributes().GetAttributes<bool>( MeshAttribute::PolygonGroup::EnableCollision );
+	const TPolygonGroupAttributeArray<bool>& PolygonGroupCastShadow = MeshDescription->PolygonGroupAttributes().GetAttributes<bool>( MeshAttribute::PolygonGroup::CastShadow );
+
 	if( EditableMesh->IsPreviewingSubdivisions() )
 	{
 		check( EditableMesh->GetSubdivisionCount() > 0 );
@@ -593,7 +548,6 @@ void UEditableStaticMeshAdapter::OnRebuildRenderMesh( const UEditableMesh* Edita
 		int32 SectionNumber = 0;
 		for( const FPolygonGroupID PolygonGroupID : PolygonGroups.GetElementIDs() )
 		{
-			const FMeshPolygonGroup& PolygonGroup = PolygonGroups[ PolygonGroupID ];
 			const FSubdivisionLimitSection& SubdivisionSection = EditableMesh->SubdivisionLimitData.Sections[ SectionNumber ];
 
 			const int32 SectionTriangleCount = SubdivisionSection.SubdividedQuads.Num() * 2;
@@ -614,13 +568,12 @@ void UEditableStaticMeshAdapter::OnRebuildRenderMesh( const UEditableMesh* Edita
 			StaticMeshSection.MinVertexIndex = FirstSectionVertexIndex;
 			StaticMeshSection.MaxVertexIndex = FirstSectionVertexIndex + SectionTriangleCount * 3;
 
-			const int32 MaterialIndex = StaticMesh->GetMaterialIndexFromImportedMaterialSlotName( PolygonGroup.ImportedMaterialSlotName );
+			const int32 MaterialIndex = StaticMesh->GetMaterialIndexFromImportedMaterialSlotName( PolygonGroupImportedMaterialSlotNames[ PolygonGroupID ] );
 			check( MaterialIndex != INDEX_NONE );
-			check( StaticMesh->StaticMaterials[ MaterialIndex ].MaterialInterface->GetPathName() == PolygonGroup.MaterialAsset.ToString() );
 
 			StaticMeshSection.MaterialIndex = MaterialIndex;
-			StaticMeshSection.bEnableCollision = PolygonGroup.bEnableCollision;
-			StaticMeshSection.bCastShadow = PolygonGroup.bCastShadow;
+			StaticMeshSection.bEnableCollision = PolygonGroupCollision[ PolygonGroupID ];
+			StaticMeshSection.bCastShadow = PolygonGroupCastShadow[ PolygonGroupID ];
 
 			// Fill vertices
 			int32 NextVertexIndex = FirstSectionVertexIndex;
@@ -685,25 +638,31 @@ void UEditableStaticMeshAdapter::OnRebuildRenderMesh( const UEditableMesh* Edita
 		// set up vertex buffer elements
 		StaticMeshBuildVertices.SetNum( VertexInstances.GetArraySize() );
 
+		const TVertexAttributeArray<FVector>& VertexPositions = MeshDescription->VertexAttributes().GetAttributes<FVector>( MeshAttribute::Vertex::Position );
+		const TVertexInstanceAttributeArray<FVector>& VertexInstanceNormals = MeshDescription->VertexInstanceAttributes().GetAttributes<FVector>( MeshAttribute::VertexInstance::Normal );
+		const TVertexInstanceAttributeArray<FVector>& VertexInstanceTangents = MeshDescription->VertexInstanceAttributes().GetAttributes<FVector>( MeshAttribute::VertexInstance::Tangent );
+		const TVertexInstanceAttributeArray<float>& VertexInstanceBinormalSigns = MeshDescription->VertexInstanceAttributes().GetAttributes<float>( MeshAttribute::VertexInstance::BinormalSign );
+		const TVertexInstanceAttributeArray<FVector4>& VertexInstanceColors = MeshDescription->VertexInstanceAttributes().GetAttributes<FVector4>( MeshAttribute::VertexInstance::Color );
+		const TVertexInstanceAttributeIndicesArray<FVector2D>& VertexInstanceUVs = MeshDescription->VertexInstanceAttributes().GetAttributesSet<FVector2D>( MeshAttribute::VertexInstance::TextureCoordinate );
+
 		for( const FVertexInstanceID VertexInstanceID : VertexInstances.GetElementIDs() )
 		{
-			const FMeshVertexInstance& VertexInstance = VertexInstances[ VertexInstanceID ];
-
-			if( VertexInstance.Color != FColor::White )
+			const FLinearColor VertexInstanceColor( VertexInstanceColors[ VertexInstanceID ] );
+			if( VertexInstanceColor != FLinearColor::White )
 			{
 				bHasColor = true;
 			}
 
 			FStaticMeshBuildVertex& StaticMeshVertex = StaticMeshBuildVertices[ VertexInstanceID.GetValue() ];
 
-			StaticMeshVertex.Position = Vertices[ VertexInstance.VertexID ].VertexPosition;
-			StaticMeshVertex.TangentX = VertexInstance.Tangent;
-			StaticMeshVertex.TangentY = FVector::CrossProduct( VertexInstance.Normal, VertexInstance.Tangent ).GetSafeNormal() * VertexInstance.BinormalSign;
-			StaticMeshVertex.TangentZ = VertexInstance.Normal;
-			StaticMeshVertex.Color = VertexInstance.Color.ToFColor( true );
-			for( int32 UVIndex = 0; UVIndex < VertexInstance.VertexUVs.Num(); ++UVIndex )
+			StaticMeshVertex.Position = VertexPositions[ MeshDescription->GetVertexInstanceVertex( VertexInstanceID ) ];
+			StaticMeshVertex.TangentX = VertexInstanceTangents[ VertexInstanceID ];
+			StaticMeshVertex.TangentY = FVector::CrossProduct( VertexInstanceNormals[ VertexInstanceID ], VertexInstanceTangents[ VertexInstanceID ] ).GetSafeNormal() * VertexInstanceBinormalSigns[ VertexInstanceID ];
+			StaticMeshVertex.TangentZ = VertexInstanceNormals[ VertexInstanceID ];
+			StaticMeshVertex.Color = VertexInstanceColor.ToFColor( true );
+			for( int32 UVIndex = 0; UVIndex < VertexInstanceUVs.GetNumIndices(); ++UVIndex )
 			{
-				StaticMeshVertex.UVs[ UVIndex ] = VertexInstance.VertexUVs[ UVIndex ];
+				StaticMeshVertex.UVs[ UVIndex ] = VertexInstanceUVs.GetArrayForIndex( UVIndex )[ VertexInstanceID ];
 			}
 		}
 
@@ -711,7 +670,6 @@ void UEditableStaticMeshAdapter::OnRebuildRenderMesh( const UEditableMesh* Edita
 
 		for( const FPolygonGroupID PolygonGroupID : PolygonGroups.GetElementIDs() )
 		{
-			const FMeshPolygonGroup& PolygonGroup = PolygonGroups[ PolygonGroupID ];
 			FRenderingPolygonGroup& RenderingPolygonGroup = RenderingPolygonGroups[ PolygonGroupID ];
 
 			RenderingPolygonGroup.RenderingSectionIndex = StaticMeshLOD.Sections.Num();
@@ -724,12 +682,11 @@ void UEditableStaticMeshAdapter::OnRebuildRenderMesh( const UEditableMesh* Edita
 			StaticMeshSection.NumTriangles = RenderingPolygonGroup.Triangles.GetArraySize();
 			check( RenderingPolygonGroup.Triangles.GetArraySize() <= RenderingPolygonGroup.MaxTriangles );
 
-			const int32 MaterialIndex = StaticMesh->GetMaterialIndexFromImportedMaterialSlotName( PolygonGroup.ImportedMaterialSlotName );
+			const int32 MaterialIndex = StaticMesh->GetMaterialIndexFromImportedMaterialSlotName( PolygonGroupImportedMaterialSlotNames[ PolygonGroupID ] );
 			check( MaterialIndex != INDEX_NONE );
-			check( StaticMesh->StaticMaterials[ MaterialIndex ].MaterialInterface->GetPathName() == PolygonGroup.MaterialAsset.ToString() );
 			StaticMeshSection.MaterialIndex = MaterialIndex;
-			StaticMeshSection.bEnableCollision = PolygonGroup.bEnableCollision;
-			StaticMeshSection.bCastShadow = PolygonGroup.bCastShadow;
+			StaticMeshSection.bEnableCollision = PolygonGroupCollision[ PolygonGroupID ];
+			StaticMeshSection.bCastShadow = PolygonGroupCastShadow[ PolygonGroupID ];
 
 
 			if( RenderingPolygonGroup.Triangles.Num() > 0 )
@@ -1065,19 +1022,15 @@ void UEditableStaticMeshAdapter::UpdateBounds( const UEditableMesh* EditableMesh
 			FBox BoundingBox;
 			BoundingBox.Init();
 
-			// Could improve performance here if necessary:
-			// 1) cache polygon IDs per vertex (in order to quickly reject orphans) and just iterate vertex array; or
-			// 2) cache bounding box per polygon
-			// There are other cases where having polygon adjacency information (1) might be useful, so it's maybe worth considering.
-
 			const FVertexArray& Vertices = EditableMesh->GetMeshDescription()->Vertices();
-			const FVertexInstanceArray& VertexInstances = EditableMesh->GetMeshDescription()->VertexInstances();
+
+			const TVertexAttributeArray<FVector>& VertexPositions = EditableMesh->GetMeshDescription()->VertexAttributes().GetAttributes<FVector>( MeshAttribute::Vertex::Position );
 
 			for( const FVertexID VertexID : Vertices.GetElementIDs() )
 			{
 				if( !EditableMesh->IsOrphanedVertex( VertexID ) )
 				{
-					BoundingBox += Vertices[ VertexID ].VertexPosition;
+					BoundingBox += VertexPositions[ VertexID ];
 				}
 			}
 
@@ -1090,8 +1043,7 @@ void UEditableStaticMeshAdapter::UpdateBounds( const UEditableMesh* EditableMesh
 			{
 				if( !EditableMesh->IsOrphanedVertex( VertexID ) )
 				{
-					const FVector VertexPosition = Vertices[ VertexID ].VertexPosition;
-					BoundingBoxAndSphere.SphereRadius = FMath::Max( ( VertexPosition - BoundingBoxAndSphere.Origin ).Size(), BoundingBoxAndSphere.SphereRadius );
+					BoundingBoxAndSphere.SphereRadius = FMath::Max( ( VertexPositions[ VertexID ] - BoundingBoxAndSphere.Origin ).Size(), BoundingBoxAndSphere.SphereRadius );
 				}
 			}
 		}
@@ -1164,14 +1116,14 @@ void UEditableStaticMeshAdapter::UpdateCollision()
 }
 
 
-void UEditableStaticMeshAdapter::OnSetVertexAttribute( const UEditableMesh* EditableMesh, const FVertexID VertexID, const FName AttributeName, const int32 AttributeIndex, const FVector4 AttributeValue )
+void UEditableStaticMeshAdapter::OnSetVertexAttribute( const UEditableMesh* EditableMesh, const FVertexID VertexID, const FMeshElementAttributeData& Attribute )
 {
 	const FMeshVertex& Vertex = EditableMesh->GetMeshDescription()->GetVertex( VertexID );
 	FStaticMeshLODResources& StaticMeshLOD = GetStaticMeshLOD();
 
-	if( AttributeName == UEditableMeshAttribute::VertexPosition() )
+	if( Attribute.AttributeName == MeshAttribute::Vertex::Position )
 	{
-		const FVector NewVertexPosition = FVector( AttributeValue );
+		const FVector NewVertexPosition = Attribute.AttributeValue.GetValue<FVector>();
 
 		// @todo mesheditor: eventually break out subdivided mesh into a different adapter which handles things differently?
 		// (may also want different component eventually)
@@ -1211,50 +1163,56 @@ void UEditableStaticMeshAdapter::OnSetVertexAttribute( const UEditableMesh* Edit
 }
 
 
-void UEditableStaticMeshAdapter::OnSetEdgeAttribute( const UEditableMesh* EditableMesh, const FEdgeID EdgeID, const FName AttributeName, const int32 AttributeIndex, const FVector4 AttributeValue )
+void UEditableStaticMeshAdapter::OnSetEdgeAttribute( const UEditableMesh* EditableMesh, const FEdgeID EdgeID, const FMeshElementAttributeData& Attribute )
 {
 	// Nothing to do here
 }
 
 
-void UEditableStaticMeshAdapter::OnSetVertexInstanceAttribute( const UEditableMesh* EditableMesh, const FVertexInstanceID VertexInstanceID, const FName AttributeName, const int32 AttributeIndex, const FVector4 AttributeValue )
+void UEditableStaticMeshAdapter::OnSetVertexInstanceAttribute( const UEditableMesh* EditableMesh, const FVertexInstanceID VertexInstanceID, const FMeshElementAttributeData& Attribute )
 {
-	const FMeshVertexInstance& VertexInstance = EditableMesh->GetMeshDescription()->GetVertexInstance( VertexInstanceID );
+	const TAttributesSet<FVertexInstanceID>& VertexInstanceAttributes = EditableMesh->GetMeshDescription()->VertexInstanceAttributes();
 	FStaticMeshLODResources& StaticMeshLOD = GetStaticMeshLOD();
 
-	if( AttributeName == UEditableMeshAttribute::VertexNormal() ||
-		AttributeName == UEditableMeshAttribute::VertexTangent() ||
-		AttributeName == UEditableMeshAttribute::VertexBinormalSign() )
+	if( Attribute.AttributeName == MeshAttribute::VertexInstance::Normal ||
+		Attribute.AttributeName == MeshAttribute::VertexInstance::Tangent ||
+		Attribute.AttributeName == MeshAttribute::VertexInstance::BinormalSign )
 	{
 		if( !EditableMesh->IsPreviewingSubdivisions() )
 		{
+			const FVector Normal = VertexInstanceAttributes.GetAttribute<FVector>( VertexInstanceID, MeshAttribute::VertexInstance::Normal );
+			const FVector Tangent = VertexInstanceAttributes.GetAttribute<FVector>( VertexInstanceID, MeshAttribute::VertexInstance::Tangent );
+			const float BinormalSign = VertexInstanceAttributes.GetAttribute<float>( VertexInstanceID, MeshAttribute::VertexInstance::BinormalSign );
+
 			// @todo mesheditor perf: SetVertexTangents() and VertexTangentX/Y() functions actually does a bit of work to compute the basis every time. 
 			// Ideally we can get/set this stuff directly to improve performance.  This became slower after high precision basis values were added.
 			// @todo mesheditor perf: this is even more pertinent now we already have the binormal sign!
 			StaticMeshLOD.VertexBuffer.SetVertexTangents(
 				VertexInstanceID.GetValue(),
-				VertexInstance.Tangent,
-				FVector::CrossProduct( VertexInstance.Normal, VertexInstance.Tangent ).GetSafeNormal() * VertexInstance.BinormalSign,
-				VertexInstance.Normal );
+				Tangent,
+				FVector::CrossProduct( Normal, Tangent ).GetSafeNormal() * BinormalSign,
+				Normal );
 		}
 	}
-	else if( AttributeName == UEditableMeshAttribute::VertexTextureCoordinate() )
+	else if( Attribute.AttributeName == MeshAttribute::VertexInstance::TextureCoordinate )
 	{
 		if( !EditableMesh->IsPreviewingSubdivisions() )
 		{
-			check( AttributeIndex < EditableMesh->GetTextureCoordinateCount() );
-			StaticMeshLOD.VertexBuffer.SetVertexUV( VertexInstanceID.GetValue(), AttributeIndex, VertexInstance.VertexUVs[ AttributeIndex ] );
+			check( Attribute.AttributeIndex < EditableMesh->GetTextureCoordinateCount() );
+			StaticMeshLOD.VertexBuffer.SetVertexUV( VertexInstanceID.GetValue(), Attribute.AttributeIndex, Attribute.AttributeValue.GetValue<FVector2D>() );
 		}
 	}
-	else if( AttributeName == UEditableMeshAttribute::VertexColor() )
+	else if( Attribute.AttributeName == MeshAttribute::VertexInstance::Color )
 	{
 		if( !EditableMesh->IsPreviewingSubdivisions() )
 		{
-			const FColor NewColor = VertexInstance.Color.ToFColor( true );
+			const FVector4 Value = Attribute.AttributeValue.GetValue<FVector4>();
+			const FLinearColor LinearColor( Value.X, Value.Y, Value.Z, Value.W );
+			const FColor NewColor = LinearColor.ToFColor( true );
 
 			if( StaticMeshLOD.ColorVertexBuffer.GetNumVertices() != EditableMesh->GetMeshDescription()->VertexInstances().GetArraySize() )
 			{
-				if( VertexInstance.Color != FLinearColor::White )
+				if( LinearColor != FLinearColor::White )
 				{
 					// Until now, we haven't needed a vertex color buffer.
 					// Force one to be generated now that we have a non-white vertex in the mesh.
@@ -1298,17 +1256,19 @@ void UEditableStaticMeshAdapter::OnCreateVertexInstances( const UEditableMesh* E
 		static TArray<FStaticMeshBuildVertex> RenderingVerticesToAppend;
 		RenderingVerticesToAppend.SetNumUninitialized( NumNewVertexBufferRenderingVertices, false );
 
+		const TVertexAttributeArray<FVector>& VertexPositions = EditableMesh->GetMeshDescription()->VertexAttributes().GetAttributes<FVector>( MeshAttribute::Vertex::Position );
+
 		for( const FVertexInstanceID VertexInstanceID : VertexInstanceIDs )
 		{
-			const FMeshVertexInstance& VertexInstance = VertexInstances[ VertexInstanceID ];
-			const FMeshVertex& ReferencedVertex = Vertices[ VertexInstance.VertexID ];
+			const FVertexID ReferencedVertexID = EditableMesh->GetMeshDescription()->GetVertexInstanceVertex( VertexInstanceID );
 
 			const int32 NewRenderingVertexIndex = VertexInstanceID.GetValue();	// Rendering vertex indices are the same as vertex instance IDs
 
 			if( NewRenderingVertexIndex < OldVertexBufferRenderingVertexCount )
 			{
 				// Rendering vertex is within the already allocated buffer. Initialize the new vertices to some defaults
-				StaticMeshLOD.PositionVertexBuffer.VertexPosition( NewRenderingVertexIndex ) = ReferencedVertex.VertexPosition;
+				// @todo mesheditor: these defaults should come from the attributes themselves
+				StaticMeshLOD.PositionVertexBuffer.VertexPosition( NewRenderingVertexIndex ) = VertexPositions[ ReferencedVertexID ];
 				StaticMeshLOD.VertexBuffer.SetVertexTangents(
 					NewRenderingVertexIndex,
 					FVector::ZeroVector,
@@ -1333,7 +1293,7 @@ void UEditableStaticMeshAdapter::OnCreateVertexInstances( const UEditableMesh* E
 				FStaticMeshBuildVertex& RenderingVertexToAppend = RenderingVerticesToAppend[ AppendVertexNumber ];
 
 				// Initialize the new vertices to some defaults
-				RenderingVertexToAppend.Position = ReferencedVertex.VertexPosition;
+				RenderingVertexToAppend.Position = VertexPositions[ ReferencedVertexID ];
 				RenderingVertexToAppend.TangentX = FVector::ZeroVector;
 				RenderingVertexToAppend.TangentY = FVector::ZeroVector;
 				RenderingVertexToAppend.TangentZ = FVector::ZeroVector;
@@ -1396,7 +1356,6 @@ void UEditableStaticMeshAdapter::OnRetriangulatePolygons( const UEditableMesh* E
 
 		const FPolygonGroupID PolygonGroupID = RenderingPolygon.PolygonGroupID;
 
-//		const FMeshPolygonGroup& PolygonGroup = EditableMesh->PolygonGroups[ PolygonGroupID.GetValue() ];
 		FRenderingPolygonGroup& RenderingPolygonGroup = RenderingPolygonGroups[ PolygonGroupID ];
 
 		const TArray<FMeshTriangle>& Triangles = Polygon.Triangles;
@@ -1709,6 +1668,15 @@ void UEditableStaticMeshAdapter::OnCreatePolygonGroups( const UEditableMesh* Edi
 {
 	const FPolygonGroupArray& PolygonGroups = EditableMesh->GetMeshDescription()->PolygonGroups();
 
+	const UMeshDescription* MeshDescription = EditableMesh->GetMeshDescription();
+	check( MeshDescription );
+
+	const TPolygonGroupAttributeArray<UObject*>& PolygonGroupMaterialAssets = MeshDescription->PolygonGroupAttributes().GetAttributes<UObject*>( MeshAttribute::PolygonGroup::MaterialAsset );
+	const TPolygonGroupAttributeArray<FName>& PolygonGroupMaterialSlotNames = MeshDescription->PolygonGroupAttributes().GetAttributes<FName>( MeshAttribute::PolygonGroup::MaterialSlotName );
+	const TPolygonGroupAttributeArray<FName>& PolygonGroupImportedMaterialSlotNames = MeshDescription->PolygonGroupAttributes().GetAttributes<FName>( MeshAttribute::PolygonGroup::ImportedMaterialSlotName );
+	const TPolygonGroupAttributeArray<bool>& PolygonGroupCollision = MeshDescription->PolygonGroupAttributes().GetAttributes<bool>( MeshAttribute::PolygonGroup::EnableCollision );
+	const TPolygonGroupAttributeArray<bool>& PolygonGroupCastShadow = MeshDescription->PolygonGroupAttributes().GetAttributes<bool>( MeshAttribute::PolygonGroup::CastShadow );
+
 	for( const FPolygonGroupID PolygonGroupID : PolygonGroupIDs )
 	{
 		const FMeshPolygonGroup& PolygonGroup = PolygonGroups[ PolygonGroupID ];
@@ -1745,11 +1713,11 @@ void UEditableStaticMeshAdapter::OnCreatePolygonGroups( const UEditableMesh* Edi
 			StaticMeshSection.NumTriangles = 0;
 			StaticMeshSection.MinVertexIndex = 0;
 			StaticMeshSection.MaxVertexIndex = 0;
-			StaticMeshSection.bEnableCollision = PolygonGroup.bEnableCollision;
-			StaticMeshSection.bCastShadow = PolygonGroup.bCastShadow;
+			StaticMeshSection.bEnableCollision = PolygonGroupCollision[ PolygonGroupID ];
+			StaticMeshSection.bCastShadow = PolygonGroupCastShadow[ PolygonGroupID ];
 
-			UMaterialInterface* Material = Cast<UMaterialInterface>( PolygonGroup.MaterialAsset.ResolveObject() );
-			StaticMeshSection.MaterialIndex = StaticMesh->StaticMaterials.Emplace( Material, PolygonGroup.MaterialSlotName, PolygonGroup.ImportedMaterialSlotName );
+			UMaterialInterface* Material = Cast<UMaterialInterface>( PolygonGroupMaterialAssets[ PolygonGroupID ] );
+			StaticMeshSection.MaterialIndex = StaticMesh->StaticMaterials.Emplace( Material, PolygonGroupMaterialSlotNames[ PolygonGroupID ], PolygonGroupImportedMaterialSlotNames[ PolygonGroupID ] );
 		}
 
 		// Insert the rendering polygon group for keeping track of these index buffer properties
@@ -1766,14 +1734,17 @@ void UEditableStaticMeshAdapter::OnDeletePolygonGroups( const UEditableMesh* Edi
 {
 	const FPolygonGroupArray& PolygonGroups = EditableMesh->GetMeshDescription()->PolygonGroups();
 
+	const TPolygonGroupAttributeArray<FName>& PolygonGroupMaterialSlotNames = EditableMesh->GetMeshDescription()->PolygonGroupAttributes().GetAttributes<FName>( MeshAttribute::PolygonGroup::MaterialSlotName );
+
 	for( const FPolygonGroupID PolygonGroupID : PolygonGroupIDs )
 	{
-		const FMeshPolygonGroup& PolygonGroup = PolygonGroups[ PolygonGroupID ];
 		FRenderingPolygonGroup& RenderingPolygonGroup = RenderingPolygonGroups[ PolygonGroupID ];
 
+		const FName SlotName = PolygonGroupMaterialSlotNames[ PolygonGroupID ];
+
 		// Remove material slot associated with section
 		const int32 MaterialIndex = StaticMesh->StaticMaterials.IndexOfByPredicate(
-			[ &PolygonGroup ]( const FStaticMaterial& StaticMaterial ) { return StaticMaterial.MaterialSlotName == PolygonGroup.MaterialSlotName; }
+			[ &SlotName ]( const FStaticMaterial& StaticMaterial ) { return StaticMaterial.MaterialSlotName == SlotName; }
 		);
 		StaticMesh->StaticMaterials.RemoveAt( MaterialIndex );
 
diff --git a/Engine/Plugins/Runtime/EditableMesh/Source/EditableMesh/EditableStaticMeshAdapter.h b/Engine/Plugins/Runtime/EditableMesh/Source/EditableMesh/EditableStaticMeshAdapter.h
--- a/Engine/Plugins/Runtime/EditableMesh/Source/EditableMesh/EditableStaticMeshAdapter.h
+++ b/Engine/Plugins/Runtime/EditableMesh/Source/EditableMesh/EditableStaticMeshAdapter.h
@@ -141,17 +141,19 @@ public:
 	virtual void OnCreateEmptyVertexRange( const UEditableMesh* EditableMesh, const TArray<FVertexID>& VertexIDs ) override;
 	virtual void OnCreateVertices( const UEditableMesh* EditableMesh, const TArray<FVertexID>& VertexIDs ) override;
 	virtual void OnCreateVertexInstances( const UEditableMesh* EditableMesh, const TArray<FVertexInstanceID>& VertexInstanceIDs ) override;
-	virtual void OnSetVertexAttribute( const UEditableMesh* EditableMesh, const FVertexID VertexID, const FName AttributeName, const int32 AttributeIndex, const FVector4 AttributeValue ) override;
-	virtual void OnSetVertexInstanceAttribute( const UEditableMesh* EditableMesh, const FVertexInstanceID VertexInstanceID, const FName AttributeName, const int32 AttributeIndex, const FVector4 AttributeValue ) override;
+	virtual void OnSetVertexAttribute( const UEditableMesh* EditableMesh, const FVertexID VertexID, const FMeshElementAttributeData& Attribute ) override;
+	virtual void OnSetVertexInstanceAttribute( const UEditableMesh* EditableMesh, const FVertexInstanceID VertexInstanceID, const FMeshElementAttributeData& Attribute ) override;
 	virtual void OnCreateEdges( const UEditableMesh* EditableMesh, const TArray<FEdgeID>& EdgeIDs ) override;
 	virtual void OnDeleteEdges( const UEditableMesh* EditableMesh, const TArray<FEdgeID>& EdgeIDs ) override;
 	virtual void OnSetEdgesVertices( const UEditableMesh* EditableMesh, const TArray<FEdgeID>& EdgeIDs ) override;
-	virtual void OnSetEdgeAttribute( const UEditableMesh* EditableMesh, const FEdgeID EdgeID, const FName AttributeName, const int32 AttributeIndex, const FVector4 AttributeValue ) override;
+	virtual void OnSetEdgeAttribute( const UEditableMesh* EditableMesh, const FEdgeID EdgeID, const FMeshElementAttributeData& Attribute ) override;
 	virtual void OnCreatePolygons( const UEditableMesh* EditableMesh, const TArray<FPolygonID>& PolygonIDs ) override;
 	virtual void OnDeletePolygons( const UEditableMesh* EditableMesh, const TArray<FPolygonID>& PolygonIDs ) override;
 	virtual void OnChangePolygonVertexInstances( const UEditableMesh* EditableMesh, const TArray<FPolygonID>& PolygonIDs ) override;
+	virtual void OnSetPolygonAttribute( const UEditableMesh* EditableMesh, const FPolygonID PolygonID, const FMeshElementAttributeData& Attribute ) override {}
 	virtual void OnCreatePolygonGroups( const UEditableMesh* EditableMesh, const TArray<FPolygonGroupID>& PolygonGroupIDs ) override;
 	virtual void OnDeletePolygonGroups( const UEditableMesh* EditableMesh, const TArray<FPolygonGroupID>& PolygonGroupIDs ) override;
+	virtual void OnSetPolygonGroupAttribute( const UEditableMesh* EditableMesh, const FPolygonGroupID PolygonGroupID, const FMeshElementAttributeData& Attribute ) override {}
 	virtual void OnAssignPolygonsToPolygonGroups( const UEditableMesh* EditableMesh, const TArray<FPolygonGroupForPolygon>& PolygonGroupForPolygons ) override;
 	virtual void OnRetriangulatePolygons( const UEditableMesh* EditableMesh, const TArray<FPolygonID>& PolygonIDs ) override;
 
diff --git a/Engine/Plugins/Runtime/EditableMesh/Source/EditableMesh/Public/EditableMesh.h b/Engine/Plugins/Runtime/EditableMesh/Source/EditableMesh/Public/EditableMesh.h
--- a/Engine/Plugins/Runtime/EditableMesh/Source/EditableMesh/Public/EditableMesh.h
+++ b/Engine/Plugins/Runtime/EditableMesh/Source/EditableMesh/Public/EditableMesh.h
@@ -15,7 +15,7 @@
 
 class UEditableMeshAdapter;
 
-EDITABLEMESH_API DECLARE_LOG_CATEGORY_EXTERN( LogEditableMesh, Log, All );
+EDITABLEMESH_API DECLARE_LOG_CATEGORY_EXTERN( LogEditableMesh, All, All );
 
 
 // @todo mesheditor: Comment these classes and enums!
@@ -86,10 +86,6 @@ public:
 	UFUNCTION( BlueprintCallable, Category="Editable Mesh" ) UEditableMesh* RevertInstance();
 	UFUNCTION( BlueprintCallable, Category="Editable Mesh" ) void PropagateInstanceChanges();
 
-	UFUNCTION( BlueprintPure, Category="Editable Mesh" ) FVector4 GetVertexAttribute( const FVertexID VertexID, const FName AttributeName, const int32 AttributeIndex ) const;
-	UFUNCTION( BlueprintPure, Category="Editable Mesh" ) FVector4 GetVertexInstanceAttribute( const FVertexInstanceID VertexInstanceID, const FName AttributeName, const int32 AttributeIndex ) const;
-	UFUNCTION( BlueprintPure, Category="Editable Mesh" ) FVector4 GetEdgeAttribute( const FEdgeID EdgeID, const FName AttributeName, const int32 AttributeIndex ) const;
-
 	/** Returns the number of vertices in this mesh */
 	UFUNCTION( BlueprintPure, Category="Editable Mesh" )
 	int32 GetVertexCount() const;
@@ -162,14 +158,6 @@ public:
 	UFUNCTION( BlueprintPure, Category="Editable Mesh" )
 	FPolygonID GetPolygonInGroup( const FPolygonGroupID PolygonGroupID, const int32 PolygonNumber ) const;
 
-	/** Returns the material assigned to this polygon group */
-	UFUNCTION( BlueprintPure, Category="Editable Mesh" )
-	UMaterialInterface* GetPolygonGroupMaterial( const FPolygonGroupID PolygonGroupID ) const;
-
-	/** Returns the material slot name assigned to this polygon group */
-	UFUNCTION( BlueprintPure, Category="Editable Mesh" )
-	FName GetPolygonGroupMaterialSlotName( const FPolygonGroupID PolygonGroupID ) const;
-
 	/** Returns the number of polygons in this mesh */
 	UFUNCTION( BlueprintPure, Category="Editable Mesh" )
 	int32 GetPolygonCount() const;
@@ -194,10 +182,6 @@ public:
 	UFUNCTION( BlueprintPure, Category="Editable Mesh" )
 	FVertexInstanceID GetPolygonPerimeterVertexInstance( const FPolygonID PolygonID, const int32 PolygonVertexNumber ) const;
 
-	/** Returns the vertex instance attribute corresponding to the indexed vertex on the polygon's perimeter */
-	UFUNCTION( BlueprintPure, Category="Editable Mesh" )
-	FVector4 GetPolygonPerimeterVertexAttribute( const FPolygonID PolygonID, const int32 PolygonVertexNumber, const FName AttributeName, const int32 AttributeIndex ) const;
-
 	/** Returns the number of hole contours this polygon has */
 	UFUNCTION( BlueprintPure, Category="Editable Mesh" )
 	int32 GetPolygonHoleCount( const FPolygonID PolygonID ) const;
@@ -214,10 +198,6 @@ public:
 	UFUNCTION( BlueprintPure, Category="Editable Mesh" )
 	FVertexInstanceID GetPolygonHoleVertexInstance( const FPolygonID PolygonID, const int32 HoleNumber, const int32 PolygonVertexNumber ) const;
 
-	/** Returns the vertex instance attribute corresponding to the indexed vertex on the given hole contour of the polygon */
-	UFUNCTION( BlueprintPure, Category="Editable Mesh" )
-	FVector4 GetPolygonHoleVertexAttribute( const FPolygonID PolygonID, const int32 HoleNumber, const int32 PolygonVertexNumber, const FName AttributeName, const int32 AttributeIndex ) const;
-
 	/** Returns the number of triangles which make up this polygon */
 	UFUNCTION( BlueprintPure, Category="Editable Mesh" )
 	int32 GetPolygonTriangulatedTriangleCount( const FPolygonID PolygonID ) const;
@@ -233,23 +213,21 @@ protected:
 
 	FEdgeID GetPolygonContourEdge( const FMeshPolygonContour& Contour, const int32 ContourEdgeNumber, bool& bOutEdgeWindingIsReversedForPolygon ) const;
 	void GetPolygonContourEdges( const FMeshPolygonContour& Contour, TArray<FEdgeID>& OutPolygonContourEdgeIDs ) const;
-	void SetVertexAttribute( const FVertexID VertexID, const FName AttributeName, const int32 AttributeIndex, const FVector4 AttributeValue );
-	void SetVertexInstanceAttribute( const FVertexInstanceID VertexInstanceID, const FName AttributeName, const int32 AttributeIndex, const FVector4 AttributeValue );
-	void SetEdgeAttribute( const FEdgeID EdgeID, const FName AttributeName, const int32 AttributeIndex, const FVector4 AttributeValue, bool bIsUndo );
+	void SetVertexAttribute( const FVertexID VertexID, const FMeshElementAttributeData& Attribute );
+	void SetVertexInstanceAttribute( const FVertexInstanceID VertexInstanceID, const FMeshElementAttributeData& Attribute );
+	void SetEdgeAttribute( const FEdgeID EdgeID, const FMeshElementAttributeData& Attribute, bool bIsUndo );
+	void SetPolygonAttribute( const FPolygonID PolygonID, const FMeshElementAttributeData& Attribute );
+	void SetPolygonGroupAttribute( const FPolygonGroupID PolygonGroupID, const FMeshElementAttributeData& Attribute );
 	void GetVertexInstanceConnectedPolygonsInSameGroup( const FVertexInstanceID VertexInstanceID, const FPolygonID PolygonID, TArray<FPolygonID>& OutPolygonIDs ) const;
 	void SetEdgeHardness( const FEdgeID EdgeID, const bool bIsHard, const bool bIsUndo );
 	FVertexInstanceID CreateVertexInstanceForContourVertex( const FVertexAndAttributes& ContourVertex, const FPolygonID PolygonID );
-	void CreatePolygonContour( const FPolygonID PolygonID, const TArray<FVertexAndAttributes>& Contour, const EPolygonEdgeHardness PolygonEdgeHardness, TArray<FEdgeID>& OutEdgeIDs, TArray<FVertexInstanceID>& OutVertexInstanceIDs );
+	void CreatePolygonContour( const TArray<FVertexAndAttributes>& Contour, TArray<FEdgeID>& OutExistingEdgeIDs, TArray<FEdgeID>& OutEdgeIDs, TArray<UMeshDescription::FContourPoint>& OutContourPoints );
 	void BackupPolygonContour( const FMeshPolygonContour& Contour, TArray<FVertexAndAttributes>& OutVerticesAndAttributes );
-	void DeletePolygonContour( const FPolygonID PolygonID, const FMeshPolygonContour& Contour, const bool bDeleteOrphanedEdges, const bool bDeleteOrphanedVertexInstances, TArray<FEdgeID>& OrphanedEdgeIDs, TArray<FVertexInstanceID>& OrphanedVertexInstanceIDs );
-	static const TArray<FName>& GetMergeableVertexInstanceAttributes();
-	void GetMergedNamedVertexInstanceAttributeData( const TArray<FName>& AttributeNames, const FVertexInstanceID VertexInstanceID, const TArray<FMeshElementAttributeData>& AttributeValuesToMerge, TArray<FMeshElementAttributeData>& OutAttributes ) const;
 	void GetConnectedSoftEdges( const FVertexID VertexID, TArray<FEdgeID>& OutConnectedSoftEdges ) const;
 	void GetPolygonsInSameSoftEdgedGroup( const FVertexID VertexInstanceID, const FPolygonID PolygonID, TArray<FPolygonID>& OutPolygonIDs ) const;
 	void GetVertexInstancesInSameSoftEdgedGroup( const FVertexID VertexID, const FPolygonID PolygonID, const bool bPolygonNotYetInitialized, TArray<FVertexInstanceID>& OutVertexInstanceIDs ) const;
 	FVertexInstanceID GetVertexInstanceInPolygonForVertex( const FPolygonID PolygonID, const FVertexID VertexID ) const;
 	void SetPolygonContourVertexAttributes( FMeshPolygonContour& Contour, const FPolygonID PolygonID, const int32 HoleIndex, const TArray<FMeshElementAttributeList>& AttributeLists );
-	bool AreAttributeListsEqual( const TArray<FMeshElementAttributeData>& A, const TArray<FMeshElementAttributeData>& B ) const;
 	void SplitVertexInstanceInPolygons( const FVertexInstanceID VertexInstanceID, const TArray<FPolygonID>& PolygonIDs );
 	void ReplaceVertexInstanceInPolygons( const FVertexInstanceID OldVertexInstanceID, const FVertexInstanceID NewVertexInstanceID, const TArray<FPolygonID>& PolygonIDs );
 	float GetPolygonCornerAngleForVertex( const FPolygonID PolygonID, const FVertexID VertexID ) const;
@@ -361,10 +339,6 @@ public:
 	 * Statics
 	 */
 
-	UFUNCTION( BlueprintPure, Category="Editable Mesh" ) static const TArray<FName>& GetValidVertexAttributes();
-	UFUNCTION( BlueprintPure, Category="Editable Mesh" ) static const TArray<FName>& GetValidVertexInstanceAttributes();
-	UFUNCTION( BlueprintPure, Category="Editable Mesh" ) static const TArray<FName>& GetValidEdgeAttributes();
-	UFUNCTION( BlueprintPure, Category="Editable Mesh" ) static const TArray<FName>& GetValidPolygonAttributes();
 	UFUNCTION( BlueprintPure, Category="Editable Mesh" ) static inline FVertexID InvalidVertexID()
 	{
 		return FVertexID::Invalid;
@@ -398,7 +372,6 @@ public:
 		return FPolygonID( PolygonIndex );
 	}
 
-	UFUNCTION( BlueprintPure, Category="Editable Mesh" ) int32 GetMaxAttributeIndex( const FName AttributeName ) const;
 	UFUNCTION( BlueprintPure, Category="Editable Mesh" ) FPolygonGroupID GetFirstValidPolygonGroup() const;
 	UFUNCTION( BlueprintPure, Category="Editable Mesh" ) int32 GetTextureCoordinateCount() const;
 	UFUNCTION( BlueprintPure, Category="Editable Mesh" ) int32 GetSubdivisionCount() const;
@@ -435,7 +408,6 @@ public:
 	UFUNCTION( BlueprintPure, Category="Editable Mesh" ) const FSubdivisionLimitData& GetSubdivisionLimitData() const;
 	UFUNCTION( BlueprintPure, Category="Editable Mesh" ) void ComputePolygonTriangulation( const FPolygonID PolygonID, TArray<FMeshTriangle>& OutTriangles ) const;
 	UFUNCTION( BlueprintPure, Category="Editable Mesh" ) bool ComputeBarycentricWeightForPointOnPolygon( const FPolygonID PolygonID, const FVector PointOnPolygon, FMeshTriangle& OutTriangle, FVector& OutTriangleVertexWeights ) const;
-	UFUNCTION( BlueprintPure, Category="Editable Mesh" ) void ComputeTextureCoordinatesForPointOnPolygon( const FPolygonID PolygonID, const FVector PointOnPolygon, bool& bOutWereTextureCoordinatesFound, TArray<FVector4>& OutInterpolatedTextureCoordinates );
 	UFUNCTION( BlueprintPure, Category="Editable Mesh" ) void ComputePolygonsSharedEdges( const TArray<FPolygonID>& PolygonIDs, TArray<FEdgeID>& OutSharedEdgeIDs ) const;
 	UFUNCTION( BlueprintPure, Category="Editable Mesh" ) void FindPolygonLoop( const FEdgeID EdgeID, TArray<FEdgeID>& OutEdgeLoopEdgeIDs, TArray<FEdgeID>& OutFlippedEdgeIDs, TArray<FEdgeID>& OutReversedEdgeIDPathToTake, TArray<FPolygonID>& OutPolygonIDsToSplit ) const;
 	UFUNCTION( BlueprintPure, Category="Editable Mesh" ) void SearchSpatialDatabaseForPolygonsPotentiallyIntersectingLineSegment( const FVector LineSegmentStart, const FVector LineSegmentEnd, TArray<FPolygonID>& OutPolygons ) const;
@@ -535,22 +507,6 @@ public:
 	UPROPERTY()
 	UMeshDescription* MeshDescription;
 
-	/** Each editable vertex in this mesh. */
-//	TSparseArray<FMeshVertex> Vertices;
-
-	/** Sparse array of rendering vertices, that matches the vertices in the mesh vertex buffers */
-//	TSparseArray<FMeshVertexInstance> VertexInstances;
-
-	/** All editable mesh edges.  Note that some of these edges will be internal polygon edges, synthesized while
-	triangulating polygons into triangles.  Static meshes currently only support triangles. */
-//	TSparseArray<FMeshEdge> Edges;
-
-	/** All of the polygons in this mesh */
-//	TSparseArray<FMeshPolygon> Polygons;
-
-	/** All of the polygon groups in this mesh */
-//	TSparseArray<FMeshPolygonGroup> PolygonGroups;
-
 // @todo mesheditor: sort out member access. Currently StaticMesh adapter relies on accessing this stuff directly
 //protected:
 
diff --git a/Engine/Plugins/Runtime/EditableMesh/Source/EditableMesh/Public/EditableMeshAdapter.h b/Engine/Plugins/Runtime/EditableMesh/Source/EditableMesh/Public/EditableMeshAdapter.h
--- a/Engine/Plugins/Runtime/EditableMesh/Source/EditableMesh/Public/EditableMeshAdapter.h
+++ b/Engine/Plugins/Runtime/EditableMesh/Source/EditableMesh/Public/EditableMeshAdapter.h
@@ -36,17 +36,19 @@ public:
 	virtual void OnCreateEmptyVertexRange( const UEditableMesh* EditableMesh, const TArray<FVertexID>& VertexIDs ) PURE_VIRTUAL(,);
 	virtual void OnCreateVertices( const UEditableMesh* EditableMesh, const TArray<FVertexID>& VertexIDs ) PURE_VIRTUAL(,);
 	virtual void OnCreateVertexInstances( const UEditableMesh* EditableMesh, const TArray<FVertexInstanceID>& VertexInstanceIDs ) PURE_VIRTUAL(,);
-	virtual void OnSetVertexAttribute( const UEditableMesh* EditableMesh, const FVertexID VertexID, const FName AttributeName, const int32 AttributeIndex, const FVector4 AttributeValue ) PURE_VIRTUAL(,);
-	virtual void OnSetVertexInstanceAttribute( const UEditableMesh* EditableMesh, const FVertexInstanceID VertexInstanceID, const FName AttributeName, const int32 AttributeIndex, const FVector4 AttributeValue ) PURE_VIRTUAL(,);
+	virtual void OnSetVertexAttribute( const UEditableMesh* EditableMesh, const FVertexID VertexID, const FMeshElementAttributeData& Attribute ) PURE_VIRTUAL(,);
+	virtual void OnSetVertexInstanceAttribute( const UEditableMesh* EditableMesh, const FVertexInstanceID VertexInstanceID, const FMeshElementAttributeData& Attribute ) PURE_VIRTUAL(,);
 	virtual void OnCreateEdges( const UEditableMesh* EditableMesh, const TArray<FEdgeID>& EdgeIDs ) PURE_VIRTUAL(,);
 	virtual void OnDeleteEdges( const UEditableMesh* EditableMesh, const TArray<FEdgeID>& EdgeIDs ) PURE_VIRTUAL(,);
 	virtual void OnSetEdgesVertices( const UEditableMesh* EditableMesh, const TArray<FEdgeID>& EdgeIDs ) PURE_VIRTUAL(,);
-	virtual void OnSetEdgeAttribute( const UEditableMesh* EditableMesh, const FEdgeID EdgeID, const FName AttributeName, const int32 AttributeIndex, const FVector4 AttributeValue ) PURE_VIRTUAL(,);
+	virtual void OnSetEdgeAttribute( const UEditableMesh* EditableMesh, const FEdgeID EdgeID, const FMeshElementAttributeData& Attribute ) PURE_VIRTUAL(,);
 	virtual void OnCreatePolygons( const UEditableMesh* EditableMesh, const TArray<FPolygonID>& PolygonIDs ) PURE_VIRTUAL(,);
 	virtual void OnDeletePolygons( const UEditableMesh* EditableMesh, const TArray<FPolygonID>& PolygonIDs ) PURE_VIRTUAL(,);
 	virtual void OnChangePolygonVertexInstances( const UEditableMesh* EditableMesh, const TArray<FPolygonID>& PolygonIDs ) PURE_VIRTUAL(,);
+	virtual void OnSetPolygonAttribute( const UEditableMesh* EditableMesh, const FPolygonID PolygonID, const FMeshElementAttributeData& Attribute ) PURE_VIRTUAL(,);
 	virtual void OnCreatePolygonGroups( const UEditableMesh* EditableMesh, const TArray<FPolygonGroupID>& PolygonGroupIDs ) PURE_VIRTUAL(,);
 	virtual void OnDeletePolygonGroups( const UEditableMesh* EditableMesh, const TArray<FPolygonGroupID>& PolygonGroupIDs ) PURE_VIRTUAL(,);
+	virtual void OnSetPolygonGroupAttribute( const UEditableMesh* EditableMesh, const FPolygonGroupID PolygonGroupID, const FMeshElementAttributeData& Attribute ) PURE_VIRTUAL(,);
 	virtual void OnAssignPolygonsToPolygonGroups( const UEditableMesh* EditableMesh, const TArray<FPolygonGroupForPolygon>& PolygonGroupForPolygons ) PURE_VIRTUAL(,);
 	virtual void OnRetriangulatePolygons( const UEditableMesh* EditableMesh, const TArray<FPolygonID>& PolygonIDs ) PURE_VIRTUAL(,);
 };
diff --git a/Engine/Plugins/Runtime/EditableMesh/Source/EditableMesh/Public/EditableMeshTypes.h b/Engine/Plugins/Runtime/EditableMesh/Source/EditableMesh/Public/EditableMeshTypes.h
--- a/Engine/Plugins/Runtime/EditableMesh/Source/EditableMesh/Public/EditableMeshTypes.h
+++ b/Engine/Plugins/Runtime/EditableMesh/Source/EditableMesh/Public/EditableMeshTypes.h
@@ -6,8 +6,11 @@
 #include "UObject/Object.h"
 #include "UObject/ScriptMacros.h"
 #include "MeshTypes.h"
+#include "MeshAttributeArray.h"
+#include "Materials/MaterialInterface.h"
 #include "EditableMeshTypes.generated.h"
 
+
 // @todo mesheditor: Move elsewhere
 namespace LogHelpers
 {
@@ -321,6 +324,161 @@ struct FSubdivisionLimitData
 };
 
 
+UENUM()
+enum class EMeshElementAttributeType : uint8
+{
+	None,
+	FVector4,
+	FVector,
+	FVector2D,
+	Float,
+	Int,
+	Bool,
+	FName,
+	UObject
+};
+
+
+USTRUCT( BlueprintType )
+struct FMeshElementAttributeValue
+{
+	// @note: This should be replaced by TVariant as soon as it becomes available.
+	// This is a very verbose and hardcoded variant type class.
+
+	GENERATED_BODY()
+
+	/** Define constructors */
+	FMeshElementAttributeValue() : Type( EMeshElementAttributeType::None ) {}
+	explicit FMeshElementAttributeValue( const FVector4 Value ) : Type( EMeshElementAttributeType::FVector4 ), Value_FVector4( Value ) {}
+	explicit FMeshElementAttributeValue( const FVector Value ) : Type( EMeshElementAttributeType::FVector ), Value_FVector( Value ) {}
+	explicit FMeshElementAttributeValue( const FVector2D Value ) : Type( EMeshElementAttributeType::FVector2D ), Value_FVector2D( Value ) {}
+	explicit FMeshElementAttributeValue( const float Value ) : Type( EMeshElementAttributeType::Float ), Value_Float( Value ) {}
+	explicit FMeshElementAttributeValue( const int Value ) : Type( EMeshElementAttributeType::Int ), Value_Int( Value ) {}
+	explicit FMeshElementAttributeValue( const bool Value ) : Type( EMeshElementAttributeType::Bool ), Value_Bool( Value ) {}
+	explicit FMeshElementAttributeValue( const FName Value ) : Type( EMeshElementAttributeType::FName ), Value_FName( Value ) {}
+	explicit FMeshElementAttributeValue( UObject* const Value ) : Type( EMeshElementAttributeType::UObject ), Value_UObject( Value ) {}
+
+	/** Returns the type of this attribute value */
+	EMeshElementAttributeType GetType() const
+	{
+		return Type;
+	}
+
+	/** Accessors for known types. There will be a runtime error if the attribute value is not of the specified type. */
+	template <typename T> T GetValue() const;
+	template <> FVector4 GetValue<FVector4>() const { check( Type == EMeshElementAttributeType::FVector4 ); return Value_FVector4; }
+	template <> FVector GetValue<FVector>() const { check( Type == EMeshElementAttributeType::FVector ); return Value_FVector; }
+	template <> FVector2D GetValue<FVector2D>() const { check( Type == EMeshElementAttributeType::FVector2D ); return Value_FVector2D; }
+	template <> float GetValue<float>() const { check( Type == EMeshElementAttributeType::Float ); return Value_Float; }
+	template <> int GetValue<int>() const { check( Type == EMeshElementAttributeType::Int ); return Value_Int; }
+	template <> bool GetValue<bool>() const { check( Type == EMeshElementAttributeType::Bool ); return Value_Bool; }
+	template <> FName GetValue<FName>() const { check( Type == EMeshElementAttributeType::FName ); return Value_FName; }
+	template <> UObject* GetValue<UObject*>() const { check( Type == EMeshElementAttributeType::UObject ); return Value_UObject; }
+
+	/**
+	 * Calls the specified polymorphic lambda, with this attribute value as its passed-in parameter.
+	 *
+	 * For example:
+	 *  AttributeValue.Visit( []( const auto& Value ) { DoSomethingWith( Value ); } );
+	 */
+	template <typename FuncType>
+	void Visit( const FuncType& Func ) const
+	{
+		switch( Type )
+		{
+			case EMeshElementAttributeType::FVector4: Func( Value_FVector4 ); break;
+			case EMeshElementAttributeType::FVector: Func( Value_FVector ); break;
+			case EMeshElementAttributeType::FVector2D: Func( Value_FVector2D ); break;
+			case EMeshElementAttributeType::Float: Func( Value_Float ); break;
+			case EMeshElementAttributeType::Int: Func( Value_Int ); break;
+			case EMeshElementAttributeType::Bool: Func( Value_Bool ); break;
+			case EMeshElementAttributeType::FName: Func( Value_FName ); break;
+			case EMeshElementAttributeType::UObject: Func( Value_UObject ); break;
+		}
+	}
+
+	/** Returns the attribute value as a string */
+	FString ToString() const
+	{
+		switch( Type )
+		{
+			case EMeshElementAttributeType::None: return FString( "<none>" );
+			case EMeshElementAttributeType::FVector4: return Value_FVector4.ToString();
+			case EMeshElementAttributeType::FVector: return Value_FVector.ToString();
+			case EMeshElementAttributeType::FVector2D: return Value_FVector2D.ToString();
+			case EMeshElementAttributeType::Float: return Lex::ToString( Value_Float );
+			case EMeshElementAttributeType::Int: return Lex::ToString( Value_Int );
+			case EMeshElementAttributeType::Bool: return Lex::ToString( Value_Bool );
+			case EMeshElementAttributeType::FName: return Value_FName.ToString();
+			case EMeshElementAttributeType::UObject: return GetNameSafe( Value_UObject );
+		}
+
+		return FString();
+	}
+
+	/** Serializer */
+	friend FArchive& operator<<( FArchive& Ar, FMeshElementAttributeValue& AttributeValue )
+	{
+		Ar << AttributeValue.Type;
+
+		switch( AttributeValue.Type )
+		{
+			case EMeshElementAttributeType::FVector4: Ar << AttributeValue.Value_FVector4; break;
+			case EMeshElementAttributeType::FVector: Ar << AttributeValue.Value_FVector; break;
+			case EMeshElementAttributeType::FVector2D: Ar << AttributeValue.Value_FVector2D; break;
+			case EMeshElementAttributeType::Float: Ar << AttributeValue.Value_Float; break;
+			case EMeshElementAttributeType::Int: Ar << AttributeValue.Value_Int; break;
+			case EMeshElementAttributeType::Bool: Ar << AttributeValue.Value_Bool; break;
+			case EMeshElementAttributeType::FName: Ar << AttributeValue.Value_FName; break;
+			case EMeshElementAttributeType::UObject: Ar << AttributeValue.Value_UObject; break;
+		}
+
+		return Ar;
+	}
+
+	/** Test equality of two attribute values */
+	friend bool operator==( const FMeshElementAttributeValue& Value1, const FMeshElementAttributeValue& Value2 )
+	{
+		if( Value1.Type != Value2.Type )
+		{
+			return false;
+		}
+
+		switch( Value1.Type )
+		{
+			case EMeshElementAttributeType::None: return true;
+			case EMeshElementAttributeType::FVector4: return Value1.Value_FVector4 == Value2.Value_FVector4;
+			case EMeshElementAttributeType::FVector: return Value1.Value_FVector == Value2.Value_FVector;
+			case EMeshElementAttributeType::FVector2D: return Value1.Value_FVector2D == Value2.Value_FVector2D;
+			case EMeshElementAttributeType::Float: return Value1.Value_Float == Value2.Value_Float;
+			case EMeshElementAttributeType::Int: return Value1.Value_Int == Value2.Value_Int;
+			case EMeshElementAttributeType::Bool: return Value1.Value_Bool == Value2.Value_Bool;
+			case EMeshElementAttributeType::FName: return Value1.Value_FName == Value2.Value_FName;
+			case EMeshElementAttributeType::UObject: return Value1.Value_UObject == Value2.Value_UObject;
+		}
+
+		return false;
+	}
+
+	friend bool operator!=( const FMeshElementAttributeValue& Value1, const FMeshElementAttributeValue& Value2 )
+	{
+		return !( Value1 == Value2 );
+	}
+
+
+private:
+	EMeshElementAttributeType Type;
+	FVector4 Value_FVector4;
+	FVector Value_FVector;
+	FVector2D Value_FVector2D;
+	float Value_Float;
+	int Value_Int;
+	bool Value_Bool;
+	FName Value_FName;
+	UObject* Value_UObject;
+};
+
+
 USTRUCT( BlueprintType )
 struct FMeshElementAttributeData
 {
@@ -336,23 +494,21 @@ struct FMeshElementAttributeData
 
 	/** The value of this attribute */
 	UPROPERTY( BlueprintReadWrite, Category="Editable Mesh" )
-	FVector4 AttributeValue;
+	FMeshElementAttributeValue AttributeValue;
 
 	/** Default constructor */
 	FMeshElementAttributeData()
 		: AttributeName( NAME_None ),
 		  AttributeIndex( 0 ),
-		  AttributeValue( 0.0f )
-	{
-	}
+		  AttributeValue()
+	{}
 
 	/** Constructor that fills in all fields */
-	FMeshElementAttributeData( const FName InitAttributeName, const int32 InitAttributeIndex, const FVector4 InitAttributeValue )
+	FMeshElementAttributeData( const FName InitAttributeName, const int32 InitAttributeIndex, const FMeshElementAttributeValue& InitAttributeValue )
 		: AttributeName( InitAttributeName ),
 		  AttributeIndex( InitAttributeIndex ),
 		  AttributeValue( InitAttributeValue )
-	{
-	}
+	{}
 
 	FString ToString() const
 	{
@@ -988,26 +1144,9 @@ struct FPolygonGroupToCreate
 {
 	GENERATED_BODY()
 
-	// @todo mesheditor: Material, bEnableCollision and bCastShadow should all be implemented as PolygonGroup attributes.
-	// This would require making Material an asset name, rather than a UObject*, which is also safer.
-
-	/** Material asset path to assign to the new polygon group */
-	UPROPERTY( BlueprintReadWrite, Category="Editable Mesh" )
-	FString MaterialAsset;
-
-	UPROPERTY( BlueprintReadWrite, Category="Editable Mesh" )
-	FName MaterialSlotName;
-
-	UPROPERTY( BlueprintReadWrite, Category="Editable Mesh" )
-	FName ImportedMaterialSlotName;
-
-	/** Whether the new polygon group should have collision enabled */
-	UPROPERTY( BlueprintReadWrite, Category="Editable Mesh" )
-	bool bEnableCollision;
-
-	/** Whether the new polygon group casts a shadow */
+	/** Attributes of this polygon group */
 	UPROPERTY( BlueprintReadWrite, Category="Editable Mesh" )
-	bool bCastShadow;
+	FMeshElementAttributeList PolygonGroupAttributes;
 
 	/** The original ID of the polygon group.  Should only be used by the undo system. */
 	UPROPERTY( BlueprintReadWrite, Category="Editable Mesh" )
@@ -1015,21 +1154,15 @@ struct FPolygonGroupToCreate
 
 	/** Default constructor */
 	FPolygonGroupToCreate()
-		: bEnableCollision( false )
-		, bCastShadow( false )
-		, OriginalPolygonGroupID( FPolygonGroupID::Invalid )
+		: OriginalPolygonGroupID( FPolygonGroupID::Invalid )
 	{
 	}
 
 	FString ToString() const
 	{
 		return FString::Printf(
-			TEXT( "Material:%s, MaterialSlotName:%s, ImportedMaterialSlotName:%s, bEnableCollision:%s, bCastShadow:%s, OriginalPolygonGroupID:%s" ),
-			*MaterialAsset,
-			*MaterialSlotName.ToString(),
-			*ImportedMaterialSlotName.ToString(),
-			*LexicalConversion::ToString( bEnableCollision ),
-			*LexicalConversion::ToString( bCastShadow ),
+			TEXT( "PolygonGroupAttributes:%s, OriginalPolygonGroupID:%s" ),
+			*PolygonGroupAttributes.ToString(),
 			*OriginalPolygonGroupID.ToString() );
 	}
 };
diff --git a/Engine/Source/Developer/MeshSimplifier/Private/QuadricMeshReduction.cpp b/Engine/Source/Developer/MeshSimplifier/Private/QuadricMeshReduction.cpp
--- a/Engine/Source/Developer/MeshSimplifier/Private/QuadricMeshReduction.cpp
+++ b/Engine/Source/Developer/MeshSimplifier/Private/QuadricMeshReduction.cpp
@@ -11,6 +11,7 @@
 #include "Features/IModularFeatures.h"
 #include "IMeshReductionInterfaces.h"
 #include "MeshDescription.h"
+#include "MeshAttributes.h"
 
 class FQuadricSimplifierMeshReductionModule : public IMeshReductionModule
 {
@@ -458,17 +459,25 @@ public:
 		int32 NumFaces = NumWedges / 3;
 		check(InMesh->Polygons().Num() == NumFaces);
 
+		const TVertexAttributeArray<FVector>& InVertexPositions = InMesh->VertexAttributes().GetAttributes<FVector>(MeshAttribute::Vertex::Position);
+		const TVertexInstanceAttributeArray<FVector>& InVertexNormals = InMesh->VertexInstanceAttributes().GetAttributes<FVector>(MeshAttribute::VertexInstance::Normal);
+		const TVertexInstanceAttributeArray<FVector>& InVertexTangents = InMesh->VertexInstanceAttributes().GetAttributes<FVector>(MeshAttribute::VertexInstance::Tangent);
+		const TVertexInstanceAttributeArray<float>& InVertexBinormalSigns = InMesh->VertexInstanceAttributes().GetAttributes<float>(MeshAttribute::VertexInstance::BinormalSign);
+		const TVertexInstanceAttributeArray<FVector4>& InVertexColors = InMesh->VertexInstanceAttributes().GetAttributes<FVector4>(MeshAttribute::VertexInstance::Color);
+		const TVertexInstanceAttributeIndicesArray<FVector2D>& InVertexUVs = InMesh->VertexInstanceAttributes().GetAttributesSet<FVector2D>(MeshAttribute::VertexInstance::TextureCoordinate);
+
 		// Process each face, build vertex buffer and index buffer
 		for (int32 FaceIndex = 0; FaceIndex < NumFaces; FaceIndex++)
 		{
-			const FMeshVertexInstance* VertexInstances[3];
-			const FMeshVertex* Vertices[3];
+			FVertexInstanceID VertexInstanceIDs[3];
+			FVertexID VertexIDs[3];
+
 			FVector Positions[3];
 			for (int32 CornerIndex = 0; CornerIndex < 3; CornerIndex++)
 			{
-				VertexInstances[CornerIndex] = &(InMesh->GetVertexInstance(FVertexInstanceID((FaceIndex * 3) + CornerIndex)));
-				Vertices[CornerIndex] = &(InMesh->GetVertex(VertexInstances[CornerIndex]->VertexID));
-				Positions[CornerIndex] = Vertices[CornerIndex]->VertexPosition;
+				VertexInstanceIDs[CornerIndex] = FVertexInstanceID((FaceIndex * 3) + CornerIndex);
+				VertexIDs[CornerIndex] = InMesh->GetVertexInstanceVertex(VertexInstanceIDs[CornerIndex]);
+				Positions[CornerIndex] = InVertexPositions[VertexIDs[CornerIndex]];
 			}
 
 			// Don't process degenerate triangles.
@@ -486,32 +495,34 @@ public:
 
 				TVertSimp< NumTexCoords > NewVert;
 
-				if (VertexInstances[CornerIndex]->ConnectedPolygons.Num() > 0)
+				const TArray<FPolygonID>& VertexInstanceConnectedPolygons = InMesh->GetVertexInstanceConnectedPolygons(VertexInstanceIDs[CornerIndex]);
+				if (VertexInstanceConnectedPolygons.Num() > 0)
 				{
-					const FMeshPolygon& MeshPolygon = InMesh->GetPolygon(VertexInstances[CornerIndex]->ConnectedPolygons[0]);
-					NewVert.MaterialIndex = MeshPolygon.PolygonGroupID.GetValue();
+					const FPolygonID PolygonID = VertexInstanceConnectedPolygons[0];
+					NewVert.MaterialIndex = InMesh->GetPolygonPolygonGroup(PolygonID).GetValue();
+					// @todo: check with Alexis: OK to conflate material index with polygon group ID? (what if there are gaps in the polygon group array?)
 				}
 
 				NewVert.Position = Positions[CornerIndex];
-				NewVert.Tangents[0] = VertexInstances[CornerIndex]->Tangent;
-				NewVert.Normal = VertexInstances[CornerIndex]->Normal;
+				NewVert.Tangents[0] = InVertexTangents[VertexInstanceIDs[CornerIndex]];
+				NewVert.Normal = InVertexNormals[VertexInstanceIDs[CornerIndex]];
 				NewVert.Tangents[1] = FVector(0.0f);
 				if (!NewVert.Normal.IsNearlyZero(SMALL_NUMBER) && !NewVert.Tangents[0].IsNearlyZero(SMALL_NUMBER))
 				{
-					NewVert.Tangents[1] = FVector::CrossProduct(NewVert.Normal, NewVert.Tangents[0]).GetSafeNormal() * VertexInstances[CornerIndex]->BinormalSign;
+					NewVert.Tangents[1] = FVector::CrossProduct(NewVert.Normal, NewVert.Tangents[0]).GetSafeNormal() * InVertexBinormalSigns[VertexInstanceIDs[CornerIndex]];
 				}
 
 				// Fix bad tangents
 				NewVert.Tangents[0] = NewVert.Tangents[0].ContainsNaN() ? FVector::ZeroVector : NewVert.Tangents[0];
 				NewVert.Tangents[1] = NewVert.Tangents[1].ContainsNaN() ? FVector::ZeroVector : NewVert.Tangents[1];
 				NewVert.Normal = NewVert.Normal.ContainsNaN() ? FVector::ZeroVector : NewVert.Normal;
-				NewVert.Color = VertexInstances[CornerIndex]->Color;
+				NewVert.Color = FLinearColor(InVertexColors[VertexInstanceIDs[CornerIndex]]);
 
 				for (int32 UVIndex = 0; UVIndex < NumTexCoords; UVIndex++)
 				{
-					if (VertexInstances[CornerIndex]->VertexUVs.IsValidIndex(UVIndex))
+					if (UVIndex < InVertexUVs.GetNumIndices())
 					{
-						NewVert.TexCoords[UVIndex] = VertexInstances[CornerIndex]->VertexUVs[UVIndex];
+						NewVert.TexCoords[UVIndex] = InVertexUVs.GetArrayForIndex(UVIndex)[VertexInstanceIDs[CornerIndex]];
 						InMeshNumTexCoords = FMath::Max(UVIndex, InMeshNumTexCoords);
 					}
 					else
@@ -599,7 +610,7 @@ public:
 
 			for (int32 TexCoordIndex = 0; TexCoordIndex < NumTexCoords; TexCoordIndex++)
 			{
-				if (InMesh->GetVertexInstance(FVertexInstanceID(0)).VertexUVs.Num() <= TexCoordIndex)
+				if (TexCoordIndex >= InVertexUVs.GetNumIndices())
 				{
 					TexCoordWeights[2 * TexCoordIndex + 0] = 0.0f;
 					TexCoordWeights[2 * TexCoordIndex + 1] = 0.0f;
@@ -627,6 +638,8 @@ public:
 		{
 			//Empty the destination mesh
 			//We want to keep the polygon group information, the reduce do not impact the material list
+			// @todo refactor this: PolygonGroup.Polygons shouldn't be mutated like this (it is likely to be made private soon).
+			// Prefer creating an empty mesh description, and copying the polygon group attributes over.
 			for (const FPolygonGroupID& PolygonGroupID : OutReducedMesh->PolygonGroups().GetElementIDs())
 			{
 				FMeshPolygonGroup& PolygonGroup = OutReducedMesh->GetPolygonGroup(PolygonGroupID);
@@ -637,18 +650,27 @@ public:
 			OutReducedMesh->VertexInstances().Reset();
 			OutReducedMesh->Vertices().Reset();
 
+			TVertexAttributeArray<FVector>& OutVertexPositions = OutReducedMesh->VertexAttributes().GetAttributes<FVector>(MeshAttribute::Vertex::Position);
+
 			//Fill the vertex array
 			for (int32 VertexIndex = 0; VertexIndex < (int32)NumVerts; ++VertexIndex)
 			{
 				FVertexID AddedVertexId = OutReducedMesh->CreateVertex();
-				FMeshVertex& NewVertex = OutReducedMesh->GetVertex(AddedVertexId);
-				NewVertex.VertexPosition = Verts[VertexIndex].Position;
-				NewVertex.CornerSharpness = 0.0f;
+				OutVertexPositions[AddedVertexId] = Verts[VertexIndex].Position;
 				check(AddedVertexId.GetValue() == VertexIndex);
 			}
 
 			TMap<int32, FPolygonGroupID> PolygonGroupMapping;
 
+			TVertexInstanceAttributeArray<FVector>& OutVertexNormals = OutReducedMesh->VertexInstanceAttributes().GetAttributes<FVector>(MeshAttribute::VertexInstance::Normal);
+			TVertexInstanceAttributeArray<FVector>& OutVertexTangents = OutReducedMesh->VertexInstanceAttributes().GetAttributes<FVector>(MeshAttribute::VertexInstance::Tangent);
+			TVertexInstanceAttributeArray<float>& OutVertexBinormalSigns = OutReducedMesh->VertexInstanceAttributes().GetAttributes<float>(MeshAttribute::VertexInstance::BinormalSign);
+			TVertexInstanceAttributeArray<FVector4>& OutVertexColors = OutReducedMesh->VertexInstanceAttributes().GetAttributes<FVector4>(MeshAttribute::VertexInstance::Color);
+			TVertexInstanceAttributeIndicesArray<FVector2D>& OutVertexUVs = OutReducedMesh->VertexInstanceAttributes().GetAttributesSet<FVector2D>(MeshAttribute::VertexInstance::TextureCoordinate);
+
+			//Specify the number of texture coords in this mesh description
+			OutVertexUVs.SetNumIndices(InMeshNumTexCoords);
+
 			//Vertex instances and Polygons
 			for (int32 TriangleIndex = 0; TriangleIndex < (int32)NumTris; TriangleIndex++)
 			{
@@ -666,25 +688,20 @@ public:
 					FVertexInstanceID AddedVertexInstanceId = OutReducedMesh->CreateVertexInstance(VertexID);
 					//Make sure the Added vertex instance ID is matching the expected vertex instance ID
 					check(AddedVertexInstanceId == VertexInstanceID);
-					FMeshVertexInstance& NewVertexInstance = OutReducedMesh->GetVertexInstance(AddedVertexInstanceId);
 					check(AddedVertexInstanceId.GetValue() == VertexInstanceIndex);
 
 					//NTBs information
-					NewVertexInstance.Tangent = Verts[Indexes[VertexInstanceIndex]].Tangents[0];
-					NewVertexInstance.BinormalSign = GetBasisDeterminantSign(Verts[Indexes[VertexInstanceIndex]].Tangents[0].GetSafeNormal(), Verts[Indexes[VertexInstanceIndex]].Tangents[1].GetSafeNormal(), Verts[Indexes[VertexInstanceIndex]].Normal.GetSafeNormal());
-					NewVertexInstance.Normal = Verts[Indexes[VertexInstanceIndex]].Normal;
+					OutVertexTangents[AddedVertexInstanceId] = Verts[Indexes[VertexInstanceIndex]].Tangents[0];
+					OutVertexBinormalSigns[AddedVertexInstanceId] = GetBasisDeterminantSign(Verts[Indexes[VertexInstanceIndex]].Tangents[0].GetSafeNormal(), Verts[Indexes[VertexInstanceIndex]].Tangents[1].GetSafeNormal(), Verts[Indexes[VertexInstanceIndex]].Normal.GetSafeNormal());
+					OutVertexNormals[AddedVertexInstanceId] = Verts[Indexes[VertexInstanceIndex]].Normal;
 
 					//Vertex Color
-					NewVertexInstance.Color = Verts[Indexes[VertexInstanceIndex]].Color;
+					OutVertexColors[AddedVertexInstanceId] = Verts[Indexes[VertexInstanceIndex]].Color;
 
 					//Texture coord
 					for (int32 TexCoordIndex = 0; TexCoordIndex < InMeshNumTexCoords; TexCoordIndex++)
 					{
-						while (NewVertexInstance.VertexUVs.Num() <= TexCoordIndex)
-						{
-							NewVertexInstance.VertexUVs.AddZeroed();
-						}
-						NewVertexInstance.VertexUVs[TexCoordIndex] = Verts[Indexes[VertexInstanceIndex]].TexCoords[TexCoordIndex];
+						OutVertexUVs.GetArrayForIndex(TexCoordIndex)[AddedVertexInstanceId] = Verts[Indexes[VertexInstanceIndex]].TexCoords[TexCoordIndex];
 					}
 				}
 				
@@ -695,14 +712,18 @@ public:
 				{
 					FPolygonGroupID PolygonGroupID(MaterialIndex);
 					check(InMesh->PolygonGroups().IsValid(PolygonGroupID));
-					const FMeshPolygonGroup& BasePolygonGroup = InMesh->GetPolygonGroup(PolygonGroupID);
 					MaterialPolygonGroupID = OutReducedMesh->PolygonGroups().Num() > MaterialIndex ? PolygonGroupID : OutReducedMesh->CreatePolygonGroup();
-					FMeshPolygonGroup& NewPolygonGroup = OutReducedMesh->GetPolygonGroup(MaterialPolygonGroupID);
-					NewPolygonGroup.bCastShadow = BasePolygonGroup.bCastShadow;
-					NewPolygonGroup.bEnableCollision = BasePolygonGroup.bEnableCollision;
-					NewPolygonGroup.ImportedMaterialSlotName = BasePolygonGroup.ImportedMaterialSlotName;
-					NewPolygonGroup.MaterialSlotName = BasePolygonGroup.MaterialSlotName;
-					NewPolygonGroup.MaterialAsset = BasePolygonGroup.MaterialAsset;
+
+					// Copy all attributes from the base polygon group to the new polygon group
+					InMesh->PolygonGroupAttributes().ForEachAttributeIndicesArray(
+						[OutReducedMesh, PolygonGroupID, MaterialPolygonGroupID](const FName Name, const auto& AttributeIndicesArray)
+						{
+							for (int32 Index = 0; Index < AttributeIndicesArray.GetNumIndices(); ++Index)
+							{
+								OutReducedMesh->PolygonGroupAttributes().SetAttribute(MaterialPolygonGroupID, Name, Index, AttributeIndicesArray.GetArrayForIndex(Index)[PolygonGroupID]);
+							}
+						}
+					);
 					PolygonGroupMapping.Add(MaterialIndex, MaterialPolygonGroupID);
 				}
 				else
@@ -727,21 +748,11 @@ public:
 						EdgeVertexIDs[0] = CornerVerticesIDs[CornerIndices[0]];
 						EdgeVertexIDs[1] = CornerVerticesIDs[CornerIndices[1]];
 
-						FEdgeID MatchEdgeId = FEdgeID::Invalid;
-						const FMeshVertex& EdgeMeshVertex = OutReducedMesh->GetVertex(EdgeVertexIDs[0]);
-						for (const FEdgeID& EdgeID : EdgeMeshVertex.ConnectedEdgeIDs)
-						{
-							FMeshEdge& Edge = OutReducedMesh->GetEdge(EdgeID);
-							if (EdgeVertexIDs[0] == Edge.VertexIDs[0] && EdgeVertexIDs[1] == Edge.VertexIDs[1] ||
-								EdgeVertexIDs[1] == Edge.VertexIDs[0] && EdgeVertexIDs[0] == Edge.VertexIDs[1])
-							{
-								MatchEdgeId = EdgeID;
-								break;
-							}
-						}
+						FEdgeID MatchEdgeId = OutReducedMesh->GetVertexPairEdge(EdgeVertexIDs[0], EdgeVertexIDs[1]);
 						if (MatchEdgeId == FEdgeID::Invalid)
 						{
 							MatchEdgeId = OutReducedMesh->CreateEdge(EdgeVertexIDs[0], EdgeVertexIDs[1]);
+							// @todo: set edge hardness?
 						}
 						ContourPoint.EdgeID = MatchEdgeId;
 						ContourPoint.VertexInstanceID = CornerInstanceIDs[CornerIndices[0]];
@@ -749,16 +760,12 @@ public:
 				}
 
 				// Insert a polygon into the mesh
-				const FPolygonID NewPolygonID = OutReducedMesh->CreatePolygon(Contours);
-				FMeshPolygonGroup& MeshPolygonGroup = OutReducedMesh->GetPolygonGroup(MaterialPolygonGroupID);
-				FMeshPolygon& NewPolygon = OutReducedMesh->GetPolygon(NewPolygonID);
-				NewPolygon.PolygonGroupID = MaterialPolygonGroupID;
-				MeshPolygonGroup.Polygons.Add(NewPolygonID);
-				int32 NewTriangleIndex = NewPolygon.Triangles.AddDefaulted();
-				FMeshTriangle& NewTriangle = NewPolygon.Triangles[NewTriangleIndex];
+				const FPolygonID NewPolygonID = OutReducedMesh->CreatePolygon(MaterialPolygonGroupID, Contours);
+				const int32 NewTriangleIndex = OutReducedMesh->GetPolygonTriangles(NewPolygonID).AddDefaulted();
+				FMeshTriangle& NewTriangle = OutReducedMesh->GetPolygonTriangles(NewPolygonID)[NewTriangleIndex];
 				for (int32 TriangleVertexIndex = 0; TriangleVertexIndex < 3; ++TriangleVertexIndex)
 				{
-					const FVertexInstanceID &VertexInstanceID = CornerInstanceIDs[TriangleVertexIndex];
+					const FVertexInstanceID VertexInstanceID = CornerInstanceIDs[TriangleVertexIndex];
 					NewTriangle.SetVertexInstanceID(TriangleVertexIndex, VertexInstanceID);
 				}
 			}
