Added MeshDescription methods to test validity of element IDs.
Tidied up EditableStaticMeshAdapter; reserves elements upfront when creating the mesh description, now uses the MeshDescription API instead of directly accessing mesh element members (which are soon to become private).

diff --git a/Engine/Plugins/Runtime/EditableMesh/Source/EditableMesh/EditableStaticMeshAdapter.cpp b/Engine/Plugins/Runtime/EditableMesh/Source/EditableMesh/EditableStaticMeshAdapter.cpp
--- a/Engine/Plugins/Runtime/EditableMesh/Source/EditableMesh/EditableStaticMeshAdapter.cpp
+++ b/Engine/Plugins/Runtime/EditableMesh/Source/EditableMesh/EditableStaticMeshAdapter.cpp
@@ -117,12 +117,12 @@ void UEditableStaticMeshAdapter::InitEditableStaticMesh( UEditableMesh* Editable
 
 				// Vertices
 				const int32 NumRenderingVertices = StaticMeshLOD.PositionVertexBuffer.GetNumVertices();
+				MeshDescription->ReserveNewVertices( NumRenderingVertices );	// possibly more than necessary, but doesn't matter
+				MeshDescription->ReserveNewVertexInstances( NumRenderingVertices );
+
 				const bool bHasColor = StaticMeshLOD.ColorVertexBuffer.GetNumVertices() > 0;
 				check( !bHasColor || StaticMeshLOD.ColorVertexBuffer.GetNumVertices() == StaticMeshLOD.VertexBuffer.GetNumVertices() );
 
-				FVertexArray& Vertices = MeshDescription->Vertices();
-				FVertexInstanceArray& VertexInstances = MeshDescription->VertexInstances();
-
 				// @todo mesheditor cleanup: This code is very similar to the static mesh build code; try to share helper structs
 				static TMultiMap<int32, int32> OverlappingRenderingVertexIndices;
 				{
@@ -218,8 +218,8 @@ void UEditableStaticMeshAdapter::InitEditableStaticMesh( UEditableMesh* Editable
 							// we've already processed this vertex position and created an editable mesh vertex for it.
 							if( OverlappingRenderingVertexIndex < RenderingVertexIndex )
 							{
-								check( VertexInstances.IsValid( OverlappingVertexInstanceID ) );
-								const FVertexID ExistingVertexID = VertexInstances[ OverlappingVertexInstanceID ].VertexID;
+								check( MeshDescription->IsVertexInstanceValid( OverlappingVertexInstanceID ) );
+								const FVertexID ExistingVertexID = MeshDescription->GetVertexInstanceVertex( OverlappingVertexInstanceID );
 
 								// We already have a unique editable vertex for this vertex instance position, so link them!
 								MeshDescription->CreateVertexInstanceWithID( VertexInstanceID, ExistingVertexID );
@@ -256,7 +256,6 @@ void UEditableStaticMeshAdapter::InitEditableStaticMesh( UEditableMesh* Editable
 							VertexInstanceUVs.GetArrayForIndex( UVIndex )[ VertexInstanceID ] = StaticMeshLOD.VertexBuffer.GetVertexUV( RenderingVertexIndex, UVIndex );
 						}
 					}
-
 				}
 
 
@@ -265,12 +264,10 @@ void UEditableStaticMeshAdapter::InitEditableStaticMesh( UEditableMesh* Editable
 				static TMultiMap<FEdgeID, TTuple<FVertexInstanceID, FVertexInstanceID>> EdgeToVertexInstancePair;
 				EdgeToVertexInstancePair.Reset();
 
-				FPolygonGroupArray& PolygonGroups = EditableMesh->GetMeshDescription()->PolygonGroups();
-				FPolygonArray& Polygons = EditableMesh->GetMeshDescription()->Polygons();
-				FEdgeArray& Edges = EditableMesh->GetMeshDescription()->Edges();
+				const uint32 NumSections = StaticMeshLOD.Sections.Num();
+				MeshDescription->ReserveNewPolygonGroups( NumSections );
 
 				// Add all polygon groups from the static mesh sections
-				const uint32 NumSections = StaticMeshLOD.Sections.Num();
 				for( uint32 RenderingSectionIndex = 0; RenderingSectionIndex < NumSections; ++RenderingSectionIndex )
 				{
 					const FStaticMeshSection& RenderingSection = StaticMeshLOD.Sections[ RenderingSectionIndex ];
@@ -295,6 +292,9 @@ void UEditableStaticMeshAdapter::InitEditableStaticMesh( UEditableMesh* Editable
 					NewRenderingPolygonGroup.MaxTriangles = NumSectionTriangles;
 					NewRenderingPolygonGroup.RenderingSectionIndex = RenderingSectionIndex;
 
+					MeshDescription->ReserveNewPolygons( NumSectionTriangles );
+					MeshDescription->ReserveNewEdges( NumSectionTriangles * 3 );	// more than required, but not a problem
+
 					for( uint32 SectionTriangleIndex = 0; SectionTriangleIndex < NumSectionTriangles; ++SectionTriangleIndex )
 					{
 						const uint32 RenderingTriangleFirstVertexIndex = FRenderingPolygonGroup::TriangleIndexToRenderingTriangleFirstIndex( RenderingSection, FTriangleID( SectionTriangleIndex ) );
@@ -304,7 +304,7 @@ void UEditableStaticMeshAdapter::InitEditableStaticMesh( UEditableMesh* Editable
 						for( uint32 TriangleVertexIndex = 0; TriangleVertexIndex < 3; ++TriangleVertexIndex )
 						{
 							TriangleVertexInstanceIDs[ TriangleVertexIndex ] = FVertexInstanceID( RenderingIndices[ RenderingTriangleFirstVertexIndex + TriangleVertexIndex ] );
-							TriangleVertexIDs[ TriangleVertexIndex ] = VertexInstances[ TriangleVertexInstanceIDs[ TriangleVertexIndex ] ].VertexID;
+							TriangleVertexIDs[ TriangleVertexIndex ] = MeshDescription->GetVertexInstanceVertex( TriangleVertexInstanceIDs[ TriangleVertexIndex ] );
 						}
 
 						// Make sure we have a valid triangle.  The triangle can be invalid because at least two if it's vertex indices 
@@ -466,12 +466,11 @@ void UEditableStaticMeshAdapter::InitializeFromEditableMesh( const UEditableMesh
 	RenderingPolygons.Reset();
 	RenderingPolygonGroups.Reset();
 
-	FPolygonGroupArray& PolygonGroups = EditableMesh->GetMeshDescription()->PolygonGroups();
-	FPolygonArray& Polygons = EditableMesh->GetMeshDescription()->Polygons();
+	const UMeshDescription* MeshDescription = EditableMesh->GetMeshDescription();
 
 	// Create all the required rendering polygon groups (initialized to 'empty', each with a unique rendering section index)
 	int32 RenderingSectionIndex = 0;
-	for( const FPolygonGroupID PolygonGroupID : PolygonGroups.GetElementIDs() )
+	for( const FPolygonGroupID PolygonGroupID : MeshDescription->PolygonGroups().GetElementIDs() )
 	{
 		RenderingPolygonGroups.Insert( PolygonGroupID );
 		FRenderingPolygonGroup& RenderingPolygonGroup = RenderingPolygonGroups[ PolygonGroupID ];
@@ -482,25 +481,23 @@ void UEditableStaticMeshAdapter::InitializeFromEditableMesh( const UEditableMesh
 	}
 
 	// Go through all the polygons, adding their triangles to the rendering polygon group
-	for( const FPolygonID PolygonID : Polygons.GetElementIDs() )
+	for( const FPolygonID PolygonID : MeshDescription->Polygons().GetElementIDs() )
 	{
-		const FMeshPolygon& Polygon = Polygons[ PolygonID ];
-
-		const FPolygonGroupID PolygonGroupID = Polygon.PolygonGroupID;
-		const FMeshPolygonGroup& PolygonGroup = PolygonGroups[ PolygonGroupID ];
+		const FPolygonGroupID PolygonGroupID = MeshDescription->GetPolygonPolygonGroup( PolygonID );
 		FRenderingPolygonGroup& RenderingPolygonGroup = RenderingPolygonGroups[ PolygonGroupID ];
 
 		RenderingPolygons.Insert( PolygonID );
 		FRenderingPolygon& RenderingPolygon = RenderingPolygons[ PolygonID ];
 		RenderingPolygon.PolygonGroupID = PolygonGroupID;
 
-		for( const FMeshTriangle& Triangle : Polygon.Triangles )
+		const TArray<FMeshTriangle>& Triangles = MeshDescription->GetPolygonTriangles( PolygonID );
+		for( const FMeshTriangle& Triangle : Triangles )
 		{
 			const FTriangleID TriangleID = RenderingPolygonGroup.Triangles.Add( Triangle );
 			RenderingPolygon.TriangulatedPolygonTriangleIndices.Add( TriangleID );
 		}
 
-		RenderingPolygonGroup.MaxTriangles += Polygon.Triangles.Num();
+		RenderingPolygonGroup.MaxTriangles += Triangles.Num();
 	}
 
 }
@@ -521,10 +518,6 @@ void UEditableStaticMeshAdapter::OnRebuildRenderMesh( const UEditableMesh* Edita
 	static TArray< uint32 > IndexBuffer;
 	IndexBuffer.Reset();
 
-	const FPolygonGroupArray& PolygonGroups = EditableMesh->GetMeshDescription()->PolygonGroups();
-
-	StaticMeshLOD.Sections.Empty( PolygonGroups.Num() );
-
 	bool bHasColor = false;
 
 	const UMeshDescription* MeshDescription = EditableMesh->GetMeshDescription();
@@ -534,6 +527,9 @@ void UEditableStaticMeshAdapter::OnRebuildRenderMesh( const UEditableMesh* Edita
 	const TPolygonGroupAttributeArray<bool>& PolygonGroupCollision = MeshDescription->PolygonGroupAttributes().GetAttributes<bool>( MeshAttribute::PolygonGroup::EnableCollision );
 	const TPolygonGroupAttributeArray<bool>& PolygonGroupCastShadow = MeshDescription->PolygonGroupAttributes().GetAttributes<bool>( MeshAttribute::PolygonGroup::CastShadow );
 
+	const int32 NumPolygonGroups = MeshDescription->PolygonGroups().Num();
+	StaticMeshLOD.Sections.Empty( NumPolygonGroups );
+
 	if( EditableMesh->IsPreviewingSubdivisions() )
 	{
 		check( EditableMesh->GetSubdivisionCount() > 0 );
@@ -543,10 +539,10 @@ void UEditableStaticMeshAdapter::OnRebuildRenderMesh( const UEditableMesh* Edita
 
 		// The Sections sparse array mirrors the SubdivisionLimitData sections array;
 		// iterate through it in parallel in order to get the material index and other section properties
-		check( PolygonGroups.Num() == EditableMesh->SubdivisionLimitData.Sections.Num() );
+		check( NumPolygonGroups == EditableMesh->SubdivisionLimitData.Sections.Num() );
 
 		int32 SectionNumber = 0;
-		for( const FPolygonGroupID PolygonGroupID : PolygonGroups.GetElementIDs() )
+		for( const FPolygonGroupID PolygonGroupID : MeshDescription->PolygonGroups().GetElementIDs() )
 		{
 			const FSubdivisionLimitSection& SubdivisionSection = EditableMesh->SubdivisionLimitData.Sections[ SectionNumber ];
 
@@ -632,11 +628,8 @@ void UEditableStaticMeshAdapter::OnRebuildRenderMesh( const UEditableMesh* Edita
 	}
 	else
 	{
-		const FVertexArray& Vertices = EditableMesh->GetMeshDescription()->Vertices();
-		const FVertexInstanceArray& VertexInstances = EditableMesh->GetMeshDescription()->VertexInstances();
-
 		// set up vertex buffer elements
-		StaticMeshBuildVertices.SetNum( VertexInstances.GetArraySize() );
+		StaticMeshBuildVertices.SetNum( MeshDescription->VertexInstances().GetArraySize() );
 
 		const TVertexAttributeArray<FVector>& VertexPositions = MeshDescription->VertexAttributes().GetAttributes<FVector>( MeshAttribute::Vertex::Position );
 		const TVertexInstanceAttributeArray<FVector>& VertexInstanceNormals = MeshDescription->VertexInstanceAttributes().GetAttributes<FVector>( MeshAttribute::VertexInstance::Normal );
@@ -645,7 +638,7 @@ void UEditableStaticMeshAdapter::OnRebuildRenderMesh( const UEditableMesh* Edita
 		const TVertexInstanceAttributeArray<FVector4>& VertexInstanceColors = MeshDescription->VertexInstanceAttributes().GetAttributes<FVector4>( MeshAttribute::VertexInstance::Color );
 		const TVertexInstanceAttributeIndicesArray<FVector2D>& VertexInstanceUVs = MeshDescription->VertexInstanceAttributes().GetAttributesSet<FVector2D>( MeshAttribute::VertexInstance::TextureCoordinate );
 
-		for( const FVertexInstanceID VertexInstanceID : VertexInstances.GetElementIDs() )
+		for( const FVertexInstanceID VertexInstanceID : MeshDescription->VertexInstances().GetElementIDs() )
 		{
 			const FLinearColor VertexInstanceColor( VertexInstanceColors[ VertexInstanceID ] );
 			if( VertexInstanceColor != FLinearColor::White )
@@ -653,6 +646,7 @@ void UEditableStaticMeshAdapter::OnRebuildRenderMesh( const UEditableMesh* Edita
 				bHasColor = true;
 			}
 
+			// @todo mesheditor: This can be much better optimized now; some of these attribute buffers can be copied in bulk directly to the render buffers.
 			FStaticMeshBuildVertex& StaticMeshVertex = StaticMeshBuildVertices[ VertexInstanceID.GetValue() ];
 
 			StaticMeshVertex.Position = VertexPositions[ MeshDescription->GetVertexInstanceVertex( VertexInstanceID ) ];
@@ -668,7 +662,7 @@ void UEditableStaticMeshAdapter::OnRebuildRenderMesh( const UEditableMesh* Edita
 
 		// Set up index buffer
 
-		for( const FPolygonGroupID PolygonGroupID : PolygonGroups.GetElementIDs() )
+		for( const FPolygonGroupID PolygonGroupID : MeshDescription->PolygonGroups().GetElementIDs() )
 		{
 			FRenderingPolygonGroup& RenderingPolygonGroup = RenderingPolygonGroups[ PolygonGroupID ];
 
@@ -873,7 +867,7 @@ void UEditableStaticMeshAdapter::OnReindexElements( const UEditableMesh* Editabl
 	for( const FPolygonGroupID PolygonGroupID : RenderingPolygonGroups.GetElementIDs() )
 	{
 		FRenderingPolygonGroup& RenderingPolygonGroup = RenderingPolygonGroups[ PolygonGroupID ];
-		const FMeshPolygonGroup& PolygonGroup = EditableMesh->GetMeshDescription()->GetPolygonGroup( PolygonGroupID );
+//		const FMeshPolygonGroup& PolygonGroup = EditableMesh->GetMeshDescription()->GetPolygonGroup( PolygonGroupID );
 
 		TSparseArray<FTriangleID> TriangleRemappings;
 		RenderingPolygonGroup.Triangles.Compact( TriangleRemappings );
@@ -892,7 +886,7 @@ void UEditableStaticMeshAdapter::OnReindexElements( const UEditableMesh* Editabl
 		RenderingPolygonGroup.MaxTriangles = RenderingPolygonGroup.Triangles.GetArraySize();
 
 		// Fix up references in referencing polygons
-		for( const FPolygonID PolygonID : PolygonGroup.Polygons )
+		for( const FPolygonID PolygonID : EditableMesh->GetMeshDescription()->GetPolygonGroupPolygons( PolygonGroupID ) )
 		{
 			FRenderingPolygon& RenderingPolygon = RenderingPolygons[ PolygonID ];
 			RenderingPolygon.PolygonGroupID = PolygonGroupID;
@@ -1022,11 +1016,11 @@ void UEditableStaticMeshAdapter::UpdateBounds( const UEditableMesh* EditableMesh
 			FBox BoundingBox;
 			BoundingBox.Init();
 
-			const FVertexArray& Vertices = EditableMesh->GetMeshDescription()->Vertices();
+			const UMeshDescription* MeshDescription = EditableMesh->GetMeshDescription();
 
 			const TVertexAttributeArray<FVector>& VertexPositions = EditableMesh->GetMeshDescription()->VertexAttributes().GetAttributes<FVector>( MeshAttribute::Vertex::Position );
 
-			for( const FVertexID VertexID : Vertices.GetElementIDs() )
+			for( const FVertexID VertexID : MeshDescription->Vertices().GetElementIDs() )
 			{
 				if( !EditableMesh->IsOrphanedVertex( VertexID ) )
 				{
@@ -1039,7 +1033,7 @@ void UEditableStaticMeshAdapter::UpdateBounds( const UEditableMesh* EditableMesh
 			// Calculate the bounding sphere, using the center of the bounding box as the origin.
 			BoundingBoxAndSphere.SphereRadius = 0.0f;
 
-			for( const FVertexID VertexID : Vertices.GetElementIDs() )
+			for( const FVertexID VertexID : MeshDescription->Vertices().GetElementIDs() )
 			{
 				if( !EditableMesh->IsOrphanedVertex( VertexID ) )
 				{
@@ -1118,9 +1112,10 @@ void UEditableStaticMeshAdapter::UpdateCollision()
 
 void UEditableStaticMeshAdapter::OnSetVertexAttribute( const UEditableMesh* EditableMesh, const FVertexID VertexID, const FMeshElementAttributeData& Attribute )
 {
-	const FMeshVertex& Vertex = EditableMesh->GetMeshDescription()->GetVertex( VertexID );
 	FStaticMeshLODResources& StaticMeshLOD = GetStaticMeshLOD();
 
+	const UMeshDescription* MeshDescription = EditableMesh->GetMeshDescription();
+
 	if( Attribute.AttributeName == MeshAttribute::Vertex::Position )
 	{
 		const FVector NewVertexPosition = Attribute.AttributeValue.GetValue<FVector>();
@@ -1132,9 +1127,9 @@ void UEditableStaticMeshAdapter::OnSetVertexAttribute( const UEditableMesh* Edit
 			const FVertexInstanceArray& VertexInstances = EditableMesh->GetMeshDescription()->VertexInstances();
 
 			// Set the vertex buffer position of all of the vertex instances for this editable vertex
-			for( const FVertexInstanceID VertexInstanceID : Vertex.VertexInstanceIDs )
+			for( const FVertexInstanceID VertexInstanceID : MeshDescription->GetVertexVertexInstances( VertexID ) )
 			{
-				check( VertexInstances.IsValid( VertexInstanceID ) );
+				check( MeshDescription->IsVertexInstanceValid( VertexInstanceID ) );
 				StaticMeshLOD.PositionVertexBuffer.VertexPosition( VertexInstanceID.GetValue() ) = NewVertexPosition;
 			}
 		}
@@ -1242,25 +1237,22 @@ void UEditableStaticMeshAdapter::OnCreateVertexInstances( const UEditableMesh* E
 {
 	if( !EditableMesh->IsPreviewingSubdivisions() )
 	{
+		const UMeshDescription* MeshDescription = EditableMesh->GetMeshDescription();
 		FStaticMeshLODResources& StaticMeshLOD = GetStaticMeshLOD();
-		const int32 NumUVs = EditableMesh->GetTextureCoordinateCount();
+		const int32 NumUVs = MeshDescription->VertexInstanceAttributes().GetAttributeIndexCount<FVector2D>( MeshAttribute::VertexInstance::TextureCoordinate );
+		const TVertexAttributeArray<FVector>& VertexPositions = MeshDescription->VertexAttributes().GetAttributes<FVector>( MeshAttribute::Vertex::Position );
 		const bool bHasColors = StaticMeshLOD.ColorVertexBuffer.GetNumVertices() > 0;
 
-		const FVertexArray& Vertices = EditableMesh->GetMeshDescription()->Vertices();
-		const FVertexInstanceArray& VertexInstances = EditableMesh->GetMeshDescription()->VertexInstances();
-
 		// Determine if we need to grow the render buffers
 		const int32 OldVertexBufferRenderingVertexCount = StaticMeshLOD.PositionVertexBuffer.GetNumVertices();
-		const int32 NumNewVertexBufferRenderingVertices = FMath::Max( 0, VertexInstances.GetArraySize() - OldVertexBufferRenderingVertexCount );
+		const int32 NumNewVertexBufferRenderingVertices = FMath::Max( 0, MeshDescription->VertexInstances().GetArraySize() - OldVertexBufferRenderingVertexCount );
 
 		static TArray<FStaticMeshBuildVertex> RenderingVerticesToAppend;
 		RenderingVerticesToAppend.SetNumUninitialized( NumNewVertexBufferRenderingVertices, false );
 
-		const TVertexAttributeArray<FVector>& VertexPositions = EditableMesh->GetMeshDescription()->VertexAttributes().GetAttributes<FVector>( MeshAttribute::Vertex::Position );
-
 		for( const FVertexInstanceID VertexInstanceID : VertexInstanceIDs )
 		{
-			const FVertexID ReferencedVertexID = EditableMesh->GetMeshDescription()->GetVertexInstanceVertex( VertexInstanceID );
+			const FVertexID ReferencedVertexID = MeshDescription->GetVertexInstanceVertex( VertexInstanceID );
 
 			const int32 NewRenderingVertexIndex = VertexInstanceID.GetValue();	// Rendering vertex indices are the same as vertex instance IDs
 
@@ -1346,19 +1338,14 @@ void UEditableStaticMeshAdapter::OnChangePolygonVertexInstances( const UEditable
 
 void UEditableStaticMeshAdapter::OnRetriangulatePolygons( const UEditableMesh* EditableMesh, const TArray<FPolygonID>& PolygonIDs )
 {
-	const FPolygonArray& Polygons = EditableMesh->GetMeshDescription()->Polygons();
-	const FPolygonGroupArray& PolygonGroups = EditableMesh->GetMeshDescription()->PolygonGroups();
+	const UMeshDescription* MeshDescription = EditableMesh->GetMeshDescription();
 
 	for( const FPolygonID PolygonID : PolygonIDs )
 	{
-		const FMeshPolygon& Polygon = Polygons[ PolygonID ];
 		FRenderingPolygon& RenderingPolygon = RenderingPolygons[ PolygonID ];
-
 		const FPolygonGroupID PolygonGroupID = RenderingPolygon.PolygonGroupID;
-
 		FRenderingPolygonGroup& RenderingPolygonGroup = RenderingPolygonGroups[ PolygonGroupID ];
-
-		const TArray<FMeshTriangle>& Triangles = Polygon.Triangles;
+		const TArray<FMeshTriangle>& Triangles = MeshDescription->GetPolygonTriangles( PolygonID );
 
 		// Check to see whether the index buffer needs to be updated
 		bool bNeedsUpdatedTriangles = false;
@@ -1666,8 +1653,6 @@ FStaticMeshLODResources& UEditableStaticMeshAdapter::GetStaticMeshLOD()
 
 void UEditableStaticMeshAdapter::OnCreatePolygonGroups( const UEditableMesh* EditableMesh, const TArray<FPolygonGroupID>& PolygonGroupIDs )
 {
-	const FPolygonGroupArray& PolygonGroups = EditableMesh->GetMeshDescription()->PolygonGroups();
-
 	const UMeshDescription* MeshDescription = EditableMesh->GetMeshDescription();
 	check( MeshDescription );
 
@@ -1679,8 +1664,6 @@ void UEditableStaticMeshAdapter::OnCreatePolygonGroups( const UEditableMesh* Edi
 
 	for( const FPolygonGroupID PolygonGroupID : PolygonGroupIDs )
 	{
-		const FMeshPolygonGroup& PolygonGroup = PolygonGroups[ PolygonGroupID ];
-
 		uint32 LODSectionIndex = 0;
 		if( !EditableMesh->IsPreviewingSubdivisions() )
 		{
@@ -1732,8 +1715,6 @@ void UEditableStaticMeshAdapter::OnCreatePolygonGroups( const UEditableMesh* Edi
 
 void UEditableStaticMeshAdapter::OnDeletePolygonGroups( const UEditableMesh* EditableMesh, const TArray<FPolygonGroupID>& PolygonGroupIDs )
 {
-	const FPolygonGroupArray& PolygonGroups = EditableMesh->GetMeshDescription()->PolygonGroups();
-
 	const TPolygonGroupAttributeArray<FName>& PolygonGroupMaterialSlotNames = EditableMesh->GetMeshDescription()->PolygonGroupAttributes().GetAttributes<FName>( MeshAttribute::PolygonGroup::MaterialSlotName );
 
 	for( const FPolygonGroupID PolygonGroupID : PolygonGroupIDs )
diff --git a/Engine/Source/Runtime/MeshDescription/Public/MeshDescription.h b/Engine/Source/Runtime/MeshDescription/Public/MeshDescription.h
--- a/Engine/Source/Runtime/MeshDescription/Public/MeshDescription.h
+++ b/Engine/Source/Runtime/MeshDescription/Public/MeshDescription.h
@@ -324,9 +324,11 @@ public:
 		VertexAttributesSet.Remove( VertexID );
 	}
 
-#if WITH_EDITORONLY_DATA
-	FString GetIdString();
-#endif
+	/** Returns whether the passed vertex ID is valid */
+	bool IsVertexValid( const FVertexID VertexID ) const
+	{
+		return VertexArray.IsValid( VertexID );
+	}
 
 	/** Reserves space for this number of new vertex instances */
 	void ReserveNewVertexInstances( const int32 NumVertexInstances )
@@ -373,6 +375,12 @@ public:
 		VertexInstanceAttributesSet.Remove( VertexInstanceID );
 	}
 
+	/** Returns whether the passed vertex instance ID is valid */
+	bool IsVertexInstanceValid( const FVertexInstanceID VertexInstanceID ) const
+	{
+		return VertexInstanceArray.IsValid( VertexInstanceID );
+	}
+
 	/** Reserves space for this number of new edges */
 	void ReserveNewEdges( const int32 NumEdges )
 	{
@@ -425,6 +433,12 @@ public:
 		EdgeAttributesSet.Remove( EdgeID );
 	}
 
+	/** Returns whether the passed edge ID is valid */
+	bool IsEdgeValid( const FEdgeID EdgeID ) const
+	{
+		return EdgeArray.IsValid( EdgeID );
+	}
+
 	/** Reserves space for this number of new polygons */
 	void ReserveNewPolygons( const int32 NumPolygons )
 	{
@@ -542,6 +556,12 @@ public:
 		PolygonAttributesSet.Remove( PolygonID );
 	}
 
+	/** Returns whether the passed polygon ID is valid */
+	bool IsPolygonValid( const FPolygonID PolygonID ) const
+	{
+		return PolygonArray.IsValid( PolygonID );
+	}
+
 	/** Reserves space for this number of new polygon groups */
 	void ReserveNewPolygonGroups( const int32 NumPolygonGroups )
 	{
@@ -578,6 +598,12 @@ public:
 		PolygonGroupAttributesSet.Remove( PolygonGroupID );
 	}
 
+	/** Returns whether the passed polygon group ID is valid */
+	bool IsPolygonGroupValid( const FPolygonGroupID PolygonGroupID ) const
+	{
+		return PolygonGroupArray.IsValid( PolygonGroupID );
+	}
+
 
 	/** Returns the edge ID defined by the two given vertex IDs, if there is one; otherwise FEdgeID::Invalid */
 	FEdgeID GetVertexPairEdge( const FVertexID VertexID0, const FVertexID VertexID1 ) const
@@ -689,6 +715,9 @@ public:
 	/** Remaps the element IDs in the mesh description according to the passed in object */
 	void Remap( const FElementIDRemappings& Remappings );
 
+#if WITH_EDITORONLY_DATA
+	FString GetIdString();
+#endif
 
 private:
 
