Fixed regression caused by changes to how materials are represented inside a polygon group.
This may be replaced in the future by a simple index, and the responsibility of managing the material assets deferred to the mesh description owner, in this case the editable mesh.

diff --git a/Engine/Plugins/Runtime/EditableMesh/Source/EditableMesh/EditableMesh.cpp b/Engine/Plugins/Runtime/EditableMesh/Source/EditableMesh/EditableMesh.cpp
--- a/Engine/Plugins/Runtime/EditableMesh/Source/EditableMesh/EditableMesh.cpp
+++ b/Engine/Plugins/Runtime/EditableMesh/Source/EditableMesh/EditableMesh.cpp
@@ -1029,7 +1029,7 @@ void UEditableMesh::SetPolygonAttribute( const FPolygonID PolygonID, const FMesh
 	for( UEditableMeshAdapter* Adapter : Adapters )
 	{
 		// @todo: implement me!
-//		Adapter->OnSetPolygonAttribute( this, PolygonID, Attribute );
+		Adapter->OnSetPolygonAttribute( this, PolygonID, Attribute );
 	}
 }
 
@@ -1041,7 +1041,7 @@ void UEditableMesh::SetPolygonGroupAttribute( const FPolygonGroupID PolygonGroup
 	for( UEditableMeshAdapter* Adapter : Adapters )
 	{
 		// @todo: implement me!
-//		Adapter->OnSetPolygonGroupAttribute( this, PolygonGroupID, Attribute );
+		Adapter->OnSetPolygonGroupAttribute( this, PolygonGroupID, Attribute );
 	}
 }
 
diff --git a/Engine/Plugins/Runtime/EditableMesh/Source/EditableMesh/EditableStaticMeshAdapter.cpp b/Engine/Plugins/Runtime/EditableMesh/Source/EditableMesh/EditableStaticMeshAdapter.cpp
--- a/Engine/Plugins/Runtime/EditableMesh/Source/EditableMesh/EditableStaticMeshAdapter.cpp
+++ b/Engine/Plugins/Runtime/EditableMesh/Source/EditableMesh/EditableStaticMeshAdapter.cpp
@@ -1713,6 +1713,36 @@ void UEditableStaticMeshAdapter::OnCreatePolygonGroups( const UEditableMesh* Edi
 }
 
 
+void UEditableStaticMeshAdapter::OnSetPolygonGroupAttribute( const UEditableMesh* EditableMesh, const FPolygonGroupID PolygonGroupID, const FMeshElementAttributeData& Attribute )
+{
+	const UMeshDescription* MeshDescription = EditableMesh->GetMeshDescription();
+
+	FRenderingPolygonGroup& RenderingPolygonGroup = RenderingPolygonGroups[ PolygonGroupID ];
+	FStaticMeshLODResources& StaticMeshLOD = GetStaticMeshLOD();
+	FStaticMeshSection& StaticMeshSection = StaticMeshLOD.Sections[ RenderingPolygonGroup.RenderingSectionIndex ];
+
+	const TPolygonGroupAttributeArray<FSoftObjectPath>& PolygonGroupMaterialAssets = MeshDescription->PolygonGroupAttributes().GetAttributes<FSoftObjectPath>( MeshAttribute::PolygonGroup::MaterialAsset );
+	const TPolygonGroupAttributeArray<FName>& PolygonGroupMaterialSlotNames = MeshDescription->PolygonGroupAttributes().GetAttributes<FName>( MeshAttribute::PolygonGroup::MaterialSlotName );
+	const TPolygonGroupAttributeArray<FName>& PolygonGroupImportedMaterialSlotNames = MeshDescription->PolygonGroupAttributes().GetAttributes<FName>( MeshAttribute::PolygonGroup::ImportedMaterialSlotName );
+
+	if( Attribute.AttributeName == MeshAttribute::PolygonGroup::MaterialAsset ||
+		Attribute.AttributeName == MeshAttribute::PolygonGroup::MaterialSlotName ||
+		Attribute.AttributeName == MeshAttribute::PolygonGroup::ImportedMaterialSlotName )
+	{
+		UMaterialInterface* Material = Cast<UMaterialInterface>( PolygonGroupMaterialAssets[ PolygonGroupID ].TryLoad() );
+		StaticMesh->StaticMaterials[ StaticMeshSection.MaterialIndex ] = FStaticMaterial( Material, PolygonGroupMaterialSlotNames[ PolygonGroupID ], PolygonGroupImportedMaterialSlotNames[ PolygonGroupID ] );
+	}
+	else if( Attribute.AttributeName == MeshAttribute::PolygonGroup::CastShadow )
+	{
+		StaticMeshSection.bCastShadow = Attribute.AttributeValue.GetValue<bool>();
+	}
+	else if( Attribute.AttributeName == MeshAttribute::PolygonGroup::EnableCollision )
+	{
+		StaticMeshSection.bEnableCollision = Attribute.AttributeValue.GetValue<bool>();
+	}
+}
+
+
 void UEditableStaticMeshAdapter::OnDeletePolygonGroups( const UEditableMesh* EditableMesh, const TArray<FPolygonGroupID>& PolygonGroupIDs )
 {
 	const TPolygonGroupAttributeArray<FName>& PolygonGroupMaterialSlotNames = EditableMesh->GetMeshDescription()->PolygonGroupAttributes().GetAttributes<FName>( MeshAttribute::PolygonGroup::MaterialSlotName );
diff --git a/Engine/Plugins/Runtime/EditableMesh/Source/EditableMesh/EditableStaticMeshAdapter.h b/Engine/Plugins/Runtime/EditableMesh/Source/EditableMesh/EditableStaticMeshAdapter.h
--- a/Engine/Plugins/Runtime/EditableMesh/Source/EditableMesh/EditableStaticMeshAdapter.h
+++ b/Engine/Plugins/Runtime/EditableMesh/Source/EditableMesh/EditableStaticMeshAdapter.h
@@ -153,7 +153,7 @@ public:
 	virtual void OnSetPolygonAttribute( const UEditableMesh* EditableMesh, const FPolygonID PolygonID, const FMeshElementAttributeData& Attribute ) override {}
 	virtual void OnCreatePolygonGroups( const UEditableMesh* EditableMesh, const TArray<FPolygonGroupID>& PolygonGroupIDs ) override;
 	virtual void OnDeletePolygonGroups( const UEditableMesh* EditableMesh, const TArray<FPolygonGroupID>& PolygonGroupIDs ) override;
-	virtual void OnSetPolygonGroupAttribute( const UEditableMesh* EditableMesh, const FPolygonGroupID PolygonGroupID, const FMeshElementAttributeData& Attribute ) override {}
+	virtual void OnSetPolygonGroupAttribute( const UEditableMesh* EditableMesh, const FPolygonGroupID PolygonGroupID, const FMeshElementAttributeData& Attribute ) override;
 	virtual void OnAssignPolygonsToPolygonGroups( const UEditableMesh* EditableMesh, const TArray<FPolygonGroupForPolygon>& PolygonGroupForPolygons ) override;
 	virtual void OnRetriangulatePolygons( const UEditableMesh* EditableMesh, const TArray<FPolygonID>& PolygonIDs ) override;
 
