Added vertex splitting when hardening an edge.
Changed method which compares attribute values for equality to allow for an epsilon in float/vector attributes.
Removed Undo hack for setting edge attributes.

diff --git a/Engine/Plugins/Runtime/EditableMesh/Source/EditableMesh/EditableMesh.cpp b/Engine/Plugins/Runtime/EditableMesh/Source/EditableMesh/EditableMesh.cpp
--- a/Engine/Plugins/Runtime/EditableMesh/Source/EditableMesh/EditableMesh.cpp
+++ b/Engine/Plugins/Runtime/EditableMesh/Source/EditableMesh/EditableMesh.cpp
@@ -207,6 +207,14 @@ static void InterpAttributesByPredicate( FMeshElementAttributeList& AttributeLis
 }
 
 
+template <typename T>
+static bool AreAttributeValuesNearlyEqual( const T& A, const T& B ) { return A == B; }
+
+static bool AreAttributeValuesNearlyEqual( const FVector& A, const FVector& B ) { return A.Equals( B ); }
+static bool AreAttributeValuesNearlyEqual( const FVector4& A, const FVector4& B ) { return A.Equals( B ); }
+static bool AreAttributeValuesNearlyEqual( const FVector2D& A, const FVector2D& B ) { return A.Equals( B ); }
+static bool AreAttributeValuesNearlyEqual( const float& A, const float& B ) { return FMath::IsNearlyEqual( A, B ); }
+
 /** Compares two elements of all attribute arrays which match the given predicate, and returns whether they are all equal or not */
 template <typename ElementIDType, typename Predicate>
 static bool CompareAttributesByPredicate( const TAttributesSet<ElementIDType>& AttributesSet, const ElementIDType ID0, const ElementIDType ID1, const Predicate& Pred )
@@ -215,12 +223,12 @@ static bool CompareAttributesByPredicate( const TAttributesSet<ElementIDType>& A
 	AttributesSet.ForEachAttributeIndicesArray(
 		[ ID0, ID1, &Pred, &bResult ]( const FName AttributeName, const auto& AttributeIndicesArray )
 		{
-			if( Pred( AttributeName, AttributeIndicesArray) )
+			if( Pred( AttributeName, AttributeIndicesArray ) )
 			{
 				for( int32 Index = 0; Index < AttributeIndicesArray.GetNumIndices(); ++Index )
 				{
 					auto& Array = AttributeIndicesArray.GetArrayForIndex( Index );
-					if( Array[ ID0 ] != Array[ ID1 ] )
+					if( !AreAttributeValuesNearlyEqual( Array[ ID0 ], Array[ ID1 ] ) )
 					{
 						bResult = false;
 						return;
@@ -1004,15 +1012,25 @@ void UEditableMesh::SetVertexInstanceAttribute( const FVertexInstanceID VertexIn
 }
 
 
-void UEditableMesh::SetEdgeAttribute( const FEdgeID EdgeID, const FMeshElementAttributeData& Attribute, bool bIsUndo )
+void UEditableMesh::SetEdgeAttribute( const FEdgeID EdgeID, const FMeshElementAttributeData& Attribute )
 {
 	ApplyAttribute( GetMeshDescription()->EdgeAttributes(), Attribute, EdgeID );
 
 	if( Attribute.AttributeName == MeshAttribute::Edge::IsHard )
 	{
-//		SetEdgeHardness( EdgeID, Attribute.AttributeValue.GetValue<bool>(), bIsUndo );
-
-		// @todo mesheditor: this needs to potentially merge/split vertex instances and rebuild normals
+		if( Attribute.AttributeValue.GetValue<bool>() )
+		{
+			// If edge is being made hard, we may need to split vertex instances.
+			static TArray<FVertexID> EdgeVertices;
+			EdgeVertices.Reset( 2 );
+			EdgeVertices.Add( GetMeshDescription()->GetEdgeVertex( EdgeID, 0 ) );
+			EdgeVertices.Add( GetMeshDescription()->GetEdgeVertex( EdgeID, 1 ) );
+			SplitVerticesIfNecessary( EdgeVertices );
+		}
+		else
+		{
+			// If edge is being made soft, mark its vertices as candidates for vertex merging at the end of the transaction.
+		}
 	}
 
 	for( UEditableMeshAdapter* Adapter : Adapters )
@@ -1040,55 +1058,46 @@ void UEditableMesh::SetPolygonGroupAttribute( const FPolygonGroupID PolygonGroup
 
 	for( UEditableMeshAdapter* Adapter : Adapters )
 	{
-		// @todo: implement me!
 		Adapter->OnSetPolygonGroupAttribute( this, PolygonGroupID, Attribute );
 	}
 }
 
 
-void UEditableMesh::GetPolygonsInSameSoftEdgedGroup( const FVertexID VertexID, const FPolygonID PolygonID, TArray<FPolygonID>& OutPolygonIDs ) const
+void UEditableMesh::GetPolygonsInSameSoftEdgedGroupAsPolygon( const FPolygonID PolygonID, const TArray<FPolygonID>& CandidatePolygonIDs, const TArray<FEdgeID>& SoftEdgeIDs, TArray<FPolygonID>& OutPolygonIDs ) const
 {
-	// The aim here is to determine which polygons form part of the same soft edged group as the polygons attached to this vertex instance.
-	// They should all contribute to the final vertex instance normal.
+	// The aim of this method is:
+	// - given a polygon ID,
+	// - given a set of candidate polygons connected to the same vertex (which should include the polygon ID),
+	// - given a set of soft edges connected to the same vertex,
+	// return the polygon IDs which form an adjacent run without crossing a hard edge.
 
 	OutPolygonIDs.Reset();
 
-	// Get all polygons connected to this vertex.
-	static TArray<FPolygonID> ConnectedPolygons;
-	GetVertexConnectedPolygons( VertexID, ConnectedPolygons );
-
-	// Cache a list of all soft edges which share this vertex.
-	// We're only interested in finding adjacent polygons which are not the other side of a hard edge.
-	static TArray<FEdgeID> ConnectedSoftEdges;
-	GetConnectedSoftEdges( VertexID, ConnectedSoftEdges );
-
-	// Iterate through adjacent polygons which contain the given vertex, but don't cross a hard edge.
 	// Maintain a list of polygon IDs to be examined. Adjacents are added to the list if suitable.
 	// Add the start poly here.
 	static TArray<FPolygonID> PolygonsToCheck;
-	PolygonsToCheck.Reset();
+	PolygonsToCheck.Reset( CandidatePolygonIDs.Num() );
 	PolygonsToCheck.Add( PolygonID );
 
-	const FEdgeArray& Edges = GetMeshDescription()->Edges();
 	int32 Index = 0;
 	while( Index < PolygonsToCheck.Num() )
 	{
 		const FPolygonID PolygonToCheck = PolygonsToCheck[ Index ];
 		Index++;
 
-		if( ConnectedPolygons.Contains( PolygonToCheck ) )
+		if( CandidatePolygonIDs.Contains( PolygonToCheck ) )
 		{
 			OutPolygonIDs.Add( PolygonToCheck );
 
 			// Now look at its adjacent polygons. If they are joined by a soft edge which includes the vertex we're interested in, we want to consider them.
 			// We take a shortcut by doing this process in reverse: we already know all the soft edges we are interested in, so check if any of them
 			// have the current polygon as an adjacent.
-			for( const FEdgeID ConnectedSoftEdge : ConnectedSoftEdges )
+			for( const FEdgeID SoftEdgeID : SoftEdgeIDs )
 			{
-				const FMeshEdge& Edge = Edges[ ConnectedSoftEdge ];
-				if( Edge.ConnectedPolygons.Contains( PolygonToCheck ) )
+				const TArray<FPolygonID>& EdgeConnectedPolygons = GetMeshDescription()->GetEdgeConnectedPolygons( SoftEdgeID );
+				if( EdgeConnectedPolygons.Contains( PolygonToCheck ) )
 				{
-					for( const FPolygonID AdjacentPolygon : Edge.ConnectedPolygons )
+					for( const FPolygonID AdjacentPolygon : EdgeConnectedPolygons )
 					{
 						// Only add new polygons which haven't yet been added to the list. This prevents circular runs of polygons triggering infinite loops.
 						PolygonsToCheck.AddUnique( AdjacentPolygon );
@@ -1100,6 +1109,24 @@ void UEditableMesh::GetPolygonsInSameSoftEdgedGroup( const FVertexID VertexID, c
 }
 
 
+void UEditableMesh::GetVertexConnectedPolygonsInSameSoftEdgedGroup( const FVertexID VertexID, const FPolygonID PolygonID, TArray<FPolygonID>& OutPolygonIDs ) const
+{
+	// The aim here is to determine which polygons form part of the same soft edged group as the polygons attached to this vertex.
+	// They should all contribute to the final vertex instance normal.
+
+	// Get all polygons connected to this vertex.
+	static TArray<FPolygonID> ConnectedPolygons;
+	GetVertexConnectedPolygons( VertexID, ConnectedPolygons );
+
+	// Cache a list of all soft edges which share this vertex.
+	// We're only interested in finding adjacent polygons which are not the other side of a hard edge.
+	static TArray<FEdgeID> ConnectedSoftEdges;
+	GetConnectedSoftEdges( VertexID, ConnectedSoftEdges );
+
+	GetPolygonsInSameSoftEdgedGroupAsPolygon( PolygonID, ConnectedPolygons, ConnectedSoftEdges, OutPolygonIDs );
+}
+
+
 void UEditableMesh::SplitVertexInstanceInPolygons( const FVertexInstanceID VertexInstanceID, const TArray<FPolygonID>& PolygonIDs )
 {
 	UE_LOG( LogEditableMesh, Verbose, TEXT( "SplitVertexInstanceInPolygons: %s %s" ), *VertexInstanceID.ToString(), *LogHelpers::ArrayToString( PolygonIDs ) );
@@ -1180,51 +1207,58 @@ void UEditableMesh::ReplaceVertexInstanceInPolygons( const FVertexInstanceID Old
 }
 
 
-void UEditableMesh::GetVertexInstanceConnectedPolygonsInSameGroup( const FVertexInstanceID VertexInstanceID, const FPolygonID PolygonID, TArray<FPolygonID>& OutPolygonIDs ) const
+void UEditableMesh::SplitVerticesIfNecessary( const TArray<FVertexID>& VerticesToSplit )
 {
-	// The aim here is to determine which polygons connected to this vertex instance form part of the same soft edged group as the specified polygon.
-
-	OutPolygonIDs.Reset();
-
-	const FMeshVertexInstance& VertexInstance = GetMeshDescription()->GetVertexInstance( VertexInstanceID );
-	const FVertexID VertexID = VertexInstance.VertexID;
-
-	// Cache a list of all soft edges which share this vertex.
-	// We're only interested in finding adjacent polygons which are not the other side of a hard edge.
-	static TArray<FEdgeID> ConnectedSoftEdges;
-	GetConnectedSoftEdges( VertexID, ConnectedSoftEdges );
-
-	// Iterate through adjacent polygons which contain the given vertex, but don't cross a hard edge.
-	// Maintain a list of polygon IDs to be examined. Adjacents are added to the list if suitable.
-	// Add the start poly here.
-	static TArray<FPolygonID> PolygonsToCheck;
-	PolygonsToCheck.Reset();
-	PolygonsToCheck.Add( PolygonID );
+	static TArray<FVertexInstanceToCreate> VertexInstancesToCreate;
+	VertexInstancesToCreate.Reset();
 
-	const FEdgeArray& Edges = GetMeshDescription()->Edges();
-	int32 Index = 0;
-	while( Index < PolygonsToCheck.Num() )
+	// Loop for each vertex to split
+	for( const FVertexID VertexToSplit : VerticesToSplit )
 	{
-		const FPolygonID PolygonToCheck = PolygonsToCheck[ Index ];
-		Index++;
+		// Cache a list of all soft edges which share this vertex.
+		static TArray<FEdgeID> ConnectedSoftEdges;
+		GetConnectedSoftEdges( VertexToSplit, ConnectedSoftEdges );
 
-		if( VertexInstance.ConnectedPolygons.Contains( PolygonToCheck ) )
-		{
-			OutPolygonIDs.Add( PolygonToCheck );
+		// Look at each vertex instance in turn.
+		// Take a copy because splitting them will mutate the list we are iterating.
+		static TArray<FVertexInstanceID> VertexInstanceIDs;
+		VertexInstanceIDs = GetMeshDescription()->GetVertexVertexInstances( VertexToSplit );
 
-			// Now look at its adjacent polygons. If they are joined by a soft edge which includes the vertex we're interested in, we want to consider them.
-			// We take a shortcut by doing this process in reverse: we already know all the soft edges we are interested in, so check if any of them
-			// have the current polygon as an adjacent.
-			for( const FEdgeID ConnectedSoftEdge : ConnectedSoftEdges )
+		for( const FVertexInstanceID VertexInstanceID : VertexInstanceIDs )
+		{
+			// Get the list of polygons connected to this vertex instance.
+			// We only need to potentially do something if there is more than one polygon connected.
+			const TArray<FPolygonID>& PolygonIDs = GetMeshDescription()->GetVertexInstanceConnectedPolygons( VertexInstanceID );
+			if( PolygonIDs.Num() > 1 )
 			{
-				const FMeshEdge& Edge = Edges[ ConnectedSoftEdge ];
-				if( Edge.ConnectedPolygons.Contains( PolygonToCheck ) )
+				// Take a copy of all the connected polygons.
+				// This is a list we will consume as we identify groups of polygons not separated by a hard edge.
+				static TArray<FPolygonID> PolygonIDsToCheck;
+				PolygonIDsToCheck = PolygonIDs;
+
+				bool bFirstTime = true;
+				while( PolygonIDsToCheck.Num() > 0 )
 				{
-					for( const FPolygonID AdjacentPolygon : Edge.ConnectedPolygons )
+					// For the next polygon in the array, determine all other polygons in the same soft edged group
+					static TArray<FPolygonID> PolygonsInSameSoftEdgedGroup;
+					GetPolygonsInSameSoftEdgedGroupAsPolygon( PolygonIDsToCheck[ 0 ], PolygonIDsToCheck, ConnectedSoftEdges, PolygonsInSameSoftEdgedGroup );
+
+					// Check that all polygons in the smoothing group are attached to this vertex instance, and remove them from the master list of polygons
+					// connected to this instance. If a polygon in the smoothing group is not attached to this vertex instance, it's because it's the other
+					// side of a UV seam and hence has a distinct vertex instance.
+					for( const FPolygonID PolygonInSameSoftEdgedGroup : PolygonsInSameSoftEdgedGroup )
 					{
-						// Only add new polygons which haven't yet been added to the list. This prevents circular runs of polygons triggering infinite loops.
-						PolygonsToCheck.AddUnique( AdjacentPolygon );
+						verify( PolygonIDsToCheck.Remove( PolygonInSameSoftEdgedGroup ) == 1 );
 					}
+
+					// First group which we extract: do nothing - they can keep their existing instance ID.
+					// Subsequent times round the loop, we create a new vertex instance copied from the original one, and replace connected polygon vertices with it.
+					if( !bFirstTime )
+					{
+						SplitVertexInstanceInPolygons( VertexInstanceID, PolygonsInSameSoftEdgedGroup );
+					}
+
+					bFirstTime = false;
 				}
 			}
 		}
@@ -4101,8 +4135,7 @@ void UEditableMesh::CreateEdges( const TArray<FEdgeToCreate>& EdgesToCreate, TAr
 	{
 		for( const FMeshElementAttributeData& EdgeAttribute : EdgesToCreate[ Index ].EdgeAttributes.Attributes )
 		{
-			const bool bIsUndo = false;
-			SetEdgeAttribute( OutNewEdgeIDs[ Index ], EdgeAttribute, bIsUndo );
+			SetEdgeAttribute( OutNewEdgeIDs[ Index ], EdgeAttribute );
 		}
 	}
 
@@ -4701,7 +4734,6 @@ void UEditableMesh::SetEdgesAttributes( const TArray<FAttributesForEdge>& Attrib
 		RevertInput.AttributesForEdges.Emplace();
 		FAttributesForEdge& RevertEdge = RevertInput.AttributesForEdges.Last();
 		RevertEdge.EdgeID = EdgeID;
-		RevertEdge.bIsUndo = true;
 
 		// Back up the attributes
 		BackupAttributesInList( RevertEdge.EdgeAttributes, AttributesForEdge.EdgeAttributes, GetMeshDescription()->EdgeAttributes(), EdgeID );
@@ -4709,8 +4741,7 @@ void UEditableMesh::SetEdgesAttributes( const TArray<FAttributesForEdge>& Attrib
 		for( const FMeshElementAttributeData& EdgeAttribute : AttributesForEdge.EdgeAttributes.Attributes )
 		{
 			// Set the new attribute
-			// @todo mesheditor: having to pass bIsUndo here is horrible. Find a better way.
-			SetEdgeAttribute( AttributesForEdge.EdgeID, EdgeAttribute, AttributesForEdge.bIsUndo );
+			SetEdgeAttribute( AttributesForEdge.EdgeID, EdgeAttribute );
 		}
 	}
 
@@ -6476,7 +6507,7 @@ void UEditableMesh::GenerateTangentsAndNormals()
 		static TArray<FPolygonID> AllConnectedPolygons;
 		const TArray<FPolygonID>& VertexInstanceConnectedPolygons = GetMeshDescription()->GetVertexInstanceConnectedPolygons( VertexInstanceID );
 		check( VertexInstanceConnectedPolygons.Num() > 0 );
-		GetPolygonsInSameSoftEdgedGroup( VertexID, VertexInstanceConnectedPolygons[ 0 ], AllConnectedPolygons );
+		GetVertexConnectedPolygonsInSameSoftEdgedGroup( VertexID, VertexInstanceConnectedPolygons[ 0 ], AllConnectedPolygons );
 
 		// The vertex instance normal is computed as a sum of all connected polygons' normals, weighted by the angle they make with the vertex
 		for( const FPolygonID ConnectedPolygonID : AllConnectedPolygons )
diff --git a/Engine/Plugins/Runtime/EditableMesh/Source/EditableMesh/Public/EditableMesh.h b/Engine/Plugins/Runtime/EditableMesh/Source/EditableMesh/Public/EditableMesh.h
--- a/Engine/Plugins/Runtime/EditableMesh/Source/EditableMesh/Public/EditableMesh.h
+++ b/Engine/Plugins/Runtime/EditableMesh/Source/EditableMesh/Public/EditableMesh.h
@@ -215,22 +215,23 @@ protected:
 	void GetPolygonContourEdges( const FMeshPolygonContour& Contour, TArray<FEdgeID>& OutPolygonContourEdgeIDs ) const;
 	void SetVertexAttribute( const FVertexID VertexID, const FMeshElementAttributeData& Attribute );
 	void SetVertexInstanceAttribute( const FVertexInstanceID VertexInstanceID, const FMeshElementAttributeData& Attribute );
-	void SetEdgeAttribute( const FEdgeID EdgeID, const FMeshElementAttributeData& Attribute, bool bIsUndo );
+	void SetEdgeAttribute( const FEdgeID EdgeID, const FMeshElementAttributeData& Attribute );
 	void SetPolygonAttribute( const FPolygonID PolygonID, const FMeshElementAttributeData& Attribute );
 	void SetPolygonGroupAttribute( const FPolygonGroupID PolygonGroupID, const FMeshElementAttributeData& Attribute );
-	void GetVertexInstanceConnectedPolygonsInSameGroup( const FVertexInstanceID VertexInstanceID, const FPolygonID PolygonID, TArray<FPolygonID>& OutPolygonIDs ) const;
 	void SetEdgeHardness( const FEdgeID EdgeID, const bool bIsHard, const bool bIsUndo );
 	FVertexInstanceID CreateVertexInstanceForContourVertex( const FVertexAndAttributes& ContourVertex, const FPolygonID PolygonID );
 	void CreatePolygonContour( const TArray<FVertexAndAttributes>& Contour, TArray<FEdgeID>& OutExistingEdgeIDs, TArray<FEdgeID>& OutEdgeIDs, TArray<UMeshDescription::FContourPoint>& OutContourPoints );
 	void BackupPolygonContour( const FMeshPolygonContour& Contour, TArray<FVertexAndAttributes>& OutVerticesAndAttributes );
 	void GetConnectedSoftEdges( const FVertexID VertexID, TArray<FEdgeID>& OutConnectedSoftEdges ) const;
-	void GetPolygonsInSameSoftEdgedGroup( const FVertexID VertexInstanceID, const FPolygonID PolygonID, TArray<FPolygonID>& OutPolygonIDs ) const;
+	void GetVertexConnectedPolygonsInSameSoftEdgedGroup( const FVertexID VertexInstanceID, const FPolygonID PolygonID, TArray<FPolygonID>& OutPolygonIDs ) const;
 	void GetVertexInstancesInSameSoftEdgedGroup( const FVertexID VertexID, const FPolygonID PolygonID, const bool bPolygonNotYetInitialized, TArray<FVertexInstanceID>& OutVertexInstanceIDs ) const;
 	FVertexInstanceID GetVertexInstanceInPolygonForVertex( const FPolygonID PolygonID, const FVertexID VertexID ) const;
 	void SetPolygonContourVertexAttributes( FMeshPolygonContour& Contour, const FPolygonID PolygonID, const int32 HoleIndex, const TArray<FMeshElementAttributeList>& AttributeLists );
 	void SplitVertexInstanceInPolygons( const FVertexInstanceID VertexInstanceID, const TArray<FPolygonID>& PolygonIDs );
 	void ReplaceVertexInstanceInPolygons( const FVertexInstanceID OldVertexInstanceID, const FVertexInstanceID NewVertexInstanceID, const TArray<FPolygonID>& PolygonIDs );
 	float GetPolygonCornerAngleForVertex( const FPolygonID PolygonID, const FVertexID VertexID ) const;
+	void SplitVerticesIfNecessary( const TArray<FVertexID>& VerticesToSplit );
+	void GetPolygonsInSameSoftEdgedGroupAsPolygon( const FPolygonID PolygonID, const TArray<FPolygonID>& PolygonIDsToCheck, const TArray<FEdgeID>& SoftEdgeIDs, TArray<FPolygonID>& OutPolygonIDs ) const;
 
 public:
 
diff --git a/Engine/Plugins/Runtime/EditableMesh/Source/EditableMesh/Public/EditableMeshTypes.h b/Engine/Plugins/Runtime/EditableMesh/Source/EditableMesh/Public/EditableMeshTypes.h
--- a/Engine/Plugins/Runtime/EditableMesh/Source/EditableMesh/Public/EditableMeshTypes.h
+++ b/Engine/Plugins/Runtime/EditableMesh/Source/EditableMesh/Public/EditableMeshTypes.h
@@ -905,25 +905,19 @@ struct FAttributesForEdge
 	UPROPERTY( BlueprintReadWrite, Category="Editable Mesh" )
 	FMeshElementAttributeList EdgeAttributes;
 
-	/** Whether this is being done as part of undo/redo */
-	UPROPERTY( BlueprintReadWrite, Category="Editable Mesh" )
-	bool bIsUndo;
-
 	/** Default constructor */
 	FAttributesForEdge()
 		: EdgeID( 0 ),
-		  EdgeAttributes(),
-		  bIsUndo( false )
+		  EdgeAttributes()
 	{
 	}
 
 	FString ToString() const
 	{
 		return FString::Printf(
-			TEXT( "EdgeID:%s, EdgeAttributes:%s, bIsUndo:%s" ),
+			TEXT( "EdgeID:%s, EdgeAttributes:%s" ),
 			*EdgeID.ToString(),
-			*EdgeAttributes.ToString(),
-			*LogHelpers::BoolToString( bIsUndo ) );
+			*EdgeAttributes.ToString() );
 	}
 };
 
diff --git a/Engine/Source/Runtime/MeshDescription/Public/MeshDescription.h b/Engine/Source/Runtime/MeshDescription/Public/MeshDescription.h
--- a/Engine/Source/Runtime/MeshDescription/Public/MeshDescription.h
+++ b/Engine/Source/Runtime/MeshDescription/Public/MeshDescription.h
@@ -639,6 +639,13 @@ public:
 		return EdgeArray[ EdgeID ].ConnectedPolygons;
 	}
 
+	/** Returns the vertex ID corresponding to one of the edge endpoints */
+	FVertexID GetEdgeVertex( const FEdgeID EdgeID, int32 VertexNumber ) const
+	{
+		check( VertexNumber == 0 || VertexNumber == 1 );
+		return EdgeArray[ EdgeID ].VertexIDs[ VertexNumber ];
+	}
+
 	/** Returns a pair of vertex IDs defining the edge */
 	TTuple<FVertexID, FVertexID> GetEdgeVertices( const FEdgeID EdgeID ) const
 	{
