Fixed crash on loading an EditableMesh which had editor-only adapters serialized.

diff --git a/Engine/Plugins/Runtime/EditableMesh/Source/EditableMesh/EditableMesh.cpp b/Engine/Plugins/Runtime/EditableMesh/Source/EditableMesh/EditableMesh.cpp
--- a/Engine/Plugins/Runtime/EditableMesh/Source/EditableMesh/EditableMesh.cpp
+++ b/Engine/Plugins/Runtime/EditableMesh/Source/EditableMesh/EditableMesh.cpp
@@ -345,6 +345,12 @@ void UEditableMesh::Serialize( FArchive& Ar )
 {
 	Super::Serialize( Ar );
 	Ar.UsingCustomVersion( FEditableMeshCustomVersion::GUID );
+
+	// If the array was serialized containing any editor-only transient adapters, they will appear here as null, so remove them.
+	if( Ar.IsLoading() )
+	{
+		Adapters.RemoveAll( []( UEditableMeshAdapter* Adapter ) { return Adapter == nullptr; } );
+	}
 }
 
 
