- Triangulation and normal/tangent recalculation at the end of a modification are now done as non-undoable transactions. This is because, on undo, they will always be performed at the end of the modification anyway as necessary.
- Added a new way to specify vertices whose instances should be considered for merging. This operation happens at the end of the modification.
- Removed dead code.

diff --git a/Engine/Plugins/Runtime/EditableMesh/Source/EditableMesh/EditableMesh.cpp b/Engine/Plugins/Runtime/EditableMesh/Source/EditableMesh/EditableMesh.cpp
--- a/Engine/Plugins/Runtime/EditableMesh/Source/EditableMesh/EditableMesh.cpp
+++ b/Engine/Plugins/Runtime/EditableMesh/Source/EditableMesh/EditableMesh.cpp
@@ -604,6 +604,7 @@ void UEditableMesh::StartModification( const EMeshModificationType MeshModificat
 
 		PolygonsPendingNewTangentBasis.Reset();
 		PolygonsPendingTriangulation.Reset();
+		VerticesPendingMerging.Reset();
 
 		// @todo mesheditor debug
 		// UE_LOG( LogEditableMesh, Log, TEXT( "UEditableStaticMesh::StartModification COMPLETE in %0.4fs" ), FunctionTimer.GetTime() );
@@ -619,15 +620,31 @@ void UEditableMesh::EndModification( const bool bFromUndo )
 		// UE_LOG( LogEditableMesh, Log, TEXT( "UEditableStaticMesh::EndModification START (ModType=%i): %s" ), (int32)MeshModificationType, *SubMeshAddress.ToString() );
 		// FAutoScopedDurationTimer FunctionTimer;
 
-		// Retriangulate first, as the triangulation of n-gons determines how the tangent basis is calculated
-		if( PolygonsPendingTriangulation.Num() > 0 )
+		// If there are any vertices which have been marked as candidates for having any of their instances merged, do this now
+		if( VerticesPendingMerging.Num() > 0 )
 		{
-			RetriangulatePolygons();
+			MergeVertexInstances();
 		}
 
-		if( PolygonsPendingNewTangentBasis.Num() > 0 )
+		// Now we need to retriangulate polygons and recalculate tangents/normals for any polygons and vertices which have been affected
+		// by any operation in this transaction.
+		// Turn off undo because the reciprocal operation will do this at the end of its transaction.
 		{
-			GenerateTangentsAndNormals();
+			const bool bIsUndoAllowed = bAllowUndo;
+			bAllowUndo = false;
+
+			// Retriangulate first, as the triangulation of n-gons determines how the tangent basis is calculated
+			if( PolygonsPendingTriangulation.Num() > 0 )
+			{
+				RetriangulatePolygons();
+			}
+
+			if( PolygonsPendingNewTangentBasis.Num() > 0 )
+			{
+				GenerateTangentsAndNormals();
+			}
+
+			bAllowUndo = bIsUndoAllowed;
 		}
 
 		if( CurrentModificationType == EMeshModificationType::Final || !bFromUndo )
@@ -1036,7 +1053,11 @@ void UEditableMesh::SetEdgeAttribute( const FEdgeID EdgeID, const FMeshElementAt
 		else
 		{
 			// If edge is being made soft, mark its vertices as candidates for vertex merging at the end of the transaction.
+			VerticesPendingMerging.Add( GetMeshDescription()->GetEdgeVertex( EdgeID, 0 ) );
+			VerticesPendingMerging.Add( GetMeshDescription()->GetEdgeVertex( EdgeID, 1 ) );
 		}
+
+		PolygonsPendingNewTangentBasis.Append( GetMeshDescription()->GetEdgeConnectedPolygons( EdgeID ) );
 	}
 
 	for( UEditableMeshAdapter* Adapter : Adapters )
@@ -1150,6 +1171,7 @@ void UEditableMesh::SplitVertexInstanceInPolygons( const FVertexInstanceID Verte
 	CreateVertexInstances( VertexInstancesToCreate, NewVertexInstanceIDs );
 
 	ReplaceVertexInstanceInPolygons( VertexInstanceID, NewVertexInstanceIDs[ 0 ], PolygonIDs );
+	UE_LOG( LogTemp, Log, TEXT( "Split vertex instance %d: new instance %d" ), VertexInstanceID.GetValue(), NewVertexInstanceIDs[ 0 ].GetValue() );
 
 	UE_LOG( LogEditableMesh, Verbose, TEXT( "* SplitVertexInstanceInPolygons returned" ) );
 }
@@ -1272,149 +1294,84 @@ void UEditableMesh::SplitVerticesIfNecessary( const TArray<FVertexID>& VerticesT
 }
 
 
-void UEditableMesh::SetEdgeHardness( const FEdgeID EdgeID, const bool bIsHard, const bool bIsUndo )
+void UEditableMesh::MergeVertexInstances()
 {
-	// @todo mesheditor: revisit splitting/welding of identical vertex indices
-#if 0
-	const FVertexArray& Vertices = GetMeshDescription()->Vertices();
-	const FVertexInstanceArray& VertexInstances = GetMeshDescription()->VertexInstances();
-
-	FMeshEdge& Edge = GetMeshDescription()->GetEdge( EdgeID );
+	static TArray<FVertexInstanceID> VertexInstancesToDelete;
+	VertexInstancesToDelete.Reset();
 
-	if( !bIsUndo )
+	for( const FVertexID VertexID : VerticesPendingMerging )
 	{
-		UE_LOG( LogEditableMesh, Verbose, TEXT( "SetEdgeHardness (changed): %s %s" ), *EdgeID.ToString(), *LogHelpers::BoolToString( bIsHard ) );
+		// Get all polygons connected to this vertex.
+		static TArray<FPolygonID> VertexConnectedPolygons;
+		GetVertexConnectedPolygons( VertexID, VertexConnectedPolygons );
 
-		if( bIsHard )
-		{
-			// Setting a hard edge potentially requires vertex instances to be split
-
-			for( const FVertexID VertexID : Edge.VertexIDs )
-			{
-				const FMeshVertex& Vertex = Vertices[ VertexID ];
-
-				// For each end of the edge, go through all vertex instances, determining if any need to be split.
-				// Take a copy of the vertex instances because splitting them will mutate the list we are iterating.
-				static TArray<FVertexInstanceID> VertexInstanceIDsCopy;
-				VertexInstanceIDsCopy = Vertex.VertexInstanceIDs;
-
-				// Iterate all the vertex instances for each endpoint.
-				// We are looking for vertex instances which need splitting into two different instances
-				// (because the hard edge will require separate normals for each of its polygon smoothing groups).
-				for( const FVertexInstanceID VertexInstanceID : VertexInstanceIDsCopy )
-				{
-					const FMeshVertexInstance& VertexInstance = VertexInstances[ VertexInstanceID ];
+		// Cache a list of all soft edges which share this vertex.
+		// We're only interested in merging vertex instances which are in the same soft edged group.
+		static TArray<FEdgeID> VertexConnectedSoftEdges;
+		GetConnectedSoftEdges( VertexID, VertexConnectedSoftEdges );
 
-					// Take a copy of the array of polygons connected to this instance.
-					static TArray<FPolygonID> PolygonsSharingThisVertex;
-					PolygonsSharingThisVertex = VertexInstance.ConnectedPolygons;
-					bool bFirstTime = true;
-					while( PolygonsSharingThisVertex.Num() > 0 )
-					{
-						// For the next polygon in the array, determine all other polygons in the same smoothing group which share this vertex instance.
-						static TArray<FPolygonID> PolygonsInSameSmoothingGroup;
-						GetVertexInstanceConnectedPolygonsInSameGroup( VertexInstanceID, PolygonsSharingThisVertex[ 0 ], PolygonsInSameSmoothingGroup );
-
-						// Check that all polygons in the smoothing group are attached to this vertex instance, and remove them from the master list of polygons
-						// connected to this instance. If a polygon in the smoothing group is not attached to this vertex instance, it's because it's the other
-						// side of a UV seam and hence has a distinct vertex instance.
-						for( const FPolygonID PolygonInSmoothingGroup : PolygonsInSameSmoothingGroup )
-						{
-							verify( PolygonsSharingThisVertex.Remove( PolygonInSmoothingGroup ) == 1 );
-						}
+		// Get all vertex instances of this vertex...
+		const TArray<FVertexInstanceID>& VertexInstanceIDs = GetMeshDescription()->GetVertexVertexInstances( VertexID );
 
-						// First group which we extract: do nothing - they can keep their existing instance ID.
-						// Subsequent times round the loop, we create a new vertex instance copied from the original one, and replace connected polygon vertices with it.
-						if( !bFirstTime )
-						{
-							SplitVertexInstanceInPolygons( VertexInstanceID, PolygonsInSameSmoothingGroup );
-						}
+		// ...and iterate through pairs of vertex instances, looking for potential to merge them
+		for( int32 IndexA = 0; IndexA < VertexInstanceIDs.Num() - 1; ++IndexA )
+		{
+			const FVertexInstanceID VertexInstanceIDA = VertexInstanceIDs[ IndexA ];
 
-						bFirstTime = false;
-					}
-				}
+			// If vertex instance isn't connected to any polygon, we can't deduce anything about its smoothing group, so skip to the next one.
+			// This will be the case if it is already orphaned, *or* if it has been merged into another vertex instance in an earlier iteration.
+			const TArray<FPolygonID>& ConnectedPolygonIDsA = GetMeshDescription()->GetVertexInstanceConnectedPolygons( VertexInstanceIDA );
+			if( ConnectedPolygonIDsA.Num() == 0 )
+			{
+				continue;
 			}
-		}
-		else
-		{
-			// Setting a soft edge potentially requires vertex instances to be merged
 
-			static TArray<FVertexInstanceID> VertexInstancesToDelete;
-			VertexInstancesToDelete.Reset();
+			// Determine the polygons which share this vertex which form the soft edged group which contain this vertex instance
+			static TArray<FPolygonID> PolygonIDsInSoftEdgedGroup;
+			GetPolygonsInSameSoftEdgedGroupAsPolygon( ConnectedPolygonIDsA[ 0 ], VertexConnectedPolygons, VertexConnectedSoftEdges, PolygonIDsInSoftEdgedGroup );
 
-			for( const FVertexID VertexID : Edge.VertexIDs )
+			for( int32 IndexB = IndexA + 1; IndexB < VertexInstanceIDs.Num(); ++IndexB )
 			{
-				// For each end of the edge which has been softened...
-				const FMeshVertex& Vertex = Vertices[ VertexID ];
+				const FVertexInstanceID VertexInstanceIDB = VertexInstanceIDs[ IndexB ];
 
-				// ...iterate through pairs of vertex instances, looking for potential to merge them
-				for( int32 IndexA = 0; IndexA < Vertex.VertexInstanceIDs.Num() - 1; ++IndexA )
+				// If this vertex instance has been marked for deletion in a previous iteration, skip it
+				if( VertexInstancesToDelete.Contains( VertexInstanceIDB ) )
 				{
-					const FVertexInstanceID VertexInstanceIDA = Vertex.VertexInstanceIDs[ IndexA ];
-					const FMeshVertexInstance& VertexInstanceA = VertexInstances[ VertexInstanceIDA ];
-
-					// If vertex instance isn't connected to any polygon, we can't deduce anything about its smoothing group, so skip to the next one
-					if( VertexInstanceA.ConnectedPolygons.Num() == 0 )
-					{
-						continue;
-					}
-
-					// Get mergeable set of attributes for first vertex instance
-					static TArray<FMeshElementAttributeData> AttributesA;
-					GetMergedNamedVertexInstanceAttributeData( GetMergeableVertexInstanceAttributes(), Vertex.VertexInstanceIDs[ IndexA ], TArray<FMeshElementAttributeData>(), AttributesA );
-
-					for( int32 IndexB = IndexA + 1; IndexB < Vertex.VertexInstanceIDs.Num(); ++IndexB )
-					{
-						const FVertexInstanceID VertexInstanceIDB = Vertex.VertexInstanceIDs[ IndexB ];
-
-						// If this vertex instance has been marked for deletion in a previous iteration, skip it
-						if( VertexInstancesToDelete.Contains( VertexInstanceIDB ) )
-						{
-							continue;
-						}
-
-						// Get mergeable set of attributes for second vertex instance
-						static TArray<FMeshElementAttributeData> AttributesB;
-						GetMergedNamedVertexInstanceAttributeData( GetMergeableVertexInstanceAttributes(), Vertex.VertexInstanceIDs[ IndexB ], TArray<FMeshElementAttributeData>(), AttributesB );
-
-						// If both vertex instances have equal mergeable attributes, they can potentially be merged
-						if( AreAttributeListsEqual( AttributesA, AttributesB ) )
-						{
-							const FMeshVertexInstance& VertexInstanceB = VertexInstances[ VertexInstanceIDB ];
+					continue;
+				}
 
-							// Check to see if they are in the same smoothing group, using an arbitary connected polygon in the first vertex instance as the pivot point.
-							static TArray<FVertexInstanceID> PossibleVertexInstances;
-							GetVertexInstancesInSameSoftEdgedGroup( VertexInstanceA.VertexID, VertexInstanceA.ConnectedPolygons[ 0 ], false, PossibleVertexInstances );
+				// If the vertex instances are not in the same soft edged group, skip it
+				const TArray<FPolygonID>& ConnectedPolygonIDsB = GetMeshDescription()->GetVertexInstanceConnectedPolygons( VertexInstanceIDB );
+				if( !PolygonIDsInSoftEdgedGroup.ContainsByPredicate( [ &ConnectedPolygonIDsB ]( const FPolygonID PolygonID ) { return ConnectedPolygonIDsB.Contains( PolygonID ); } ) )
+				{
+					continue;
+				}
 
-							// If the smoothing group contains both vertex instances, they can be merged.
-							if( PossibleVertexInstances.Contains( VertexInstanceIDA ) && PossibleVertexInstances.Contains( VertexInstanceIDB ) )
-							{
-								// Change occurrences of VertexInstanceB for VertexInstanceA in VertexInstanceB's connected polygons.
-								// Note, this will cause VertexInstanceA's connected polygons list to be added to (at the end).
-								// This works because we are evaluating the number of connected polygons each time round the loop.
-								ReplaceVertexInstanceInPolygons( VertexInstanceIDB, VertexInstanceIDA, VertexInstanceB.ConnectedPolygons );
-
-								// This will also cause VertexInstanceB to be disconnected from all polygons.
-								// We mark the vertex instance for deletion here, but do not delete it until the end, as to do so would interrupt iterating through vertex instances.
-								check( VertexInstanceB.ConnectedPolygons.Num() == 0 );
-								VertexInstancesToDelete.Add( VertexInstanceIDB );
-							}
-						}
-					}
+				if( CompareAttributesByPredicate(
+					GetMeshDescription()->VertexInstanceAttributes(),
+					VertexInstanceIDA,
+					VertexInstanceIDB,
+					[]( const FName AttributeName, const auto& AttributeIndicesArray ) { return EnumHasAllFlags( AttributeIndicesArray.GetFlags(), EMeshAttributeFlags::Mergeable ); }
+				  ) )
+				{
+					// Change occurrences of VertexInstanceB for VertexInstanceA in VertexInstanceB's connected polygons.
+					// Note, this will cause VertexInstanceA's connected polygons list to be added to (at the end).
+					// This works because we are evaluating the number of connected polygons each time round the loop.
+					UE_LOG( LogTemp, Log, TEXT( "Replacing vertex instance %d with vertex instance %d" ), VertexInstanceIDB.GetValue(), VertexInstanceIDA.GetValue() )
+					ReplaceVertexInstanceInPolygons( VertexInstanceIDB, VertexInstanceIDA, ConnectedPolygonIDsB );
+
+					// This will also cause VertexInstanceB to be disconnected from all polygons.
+					// We mark the vertex instance for deletion here, but do not delete it until the end, as to do so would interrupt iterating through vertex instances.
+					check( ConnectedPolygonIDsB.Num() == 0 );
+					VertexInstancesToDelete.Add( VertexInstanceIDB );
 				}
 			}
-
-			// Delete orphaned vertex instances
-			const bool bDeleteOrphanedVertices = false;
-			DeleteVertexInstances( VertexInstancesToDelete, bDeleteOrphanedVertices );
 		}
-
-		UE_LOG( LogEditableMesh, Verbose, TEXT( "* SetEdgeHardness returned" ) );
 	}
 
-	// Mark polygons adjacent to this edge as candidates for normal/tangent recalculation
-	PolygonsPendingNewTangentBasis.Append( Edge.ConnectedPolygons );
-#endif
+	// Delete orphaned vertex instances
+	const bool bDeleteOrphanedVertices = false;
+	DeleteVertexInstances( VertexInstancesToDelete, bDeleteOrphanedVertices );
 }
 
 
@@ -4922,66 +4879,6 @@ void UEditableMesh::GetConnectedSoftEdges( const FVertexID VertexID, TArray<FEdg
 }
 
 
-void UEditableMesh::GetVertexInstancesInSameSoftEdgedGroup( const FVertexID VertexID, const FPolygonID PolygonID, const bool bPolygonNotYetInitialized, TArray<FVertexInstanceID>& OutVertexInstanceIDs ) const
-{
-	OutVertexInstanceIDs.Reset();
-
-	const FEdgeArray& Edges = GetMeshDescription()->Edges();
-
-	// Cache a list of all soft edges which share this vertex.
-	// We're not interested in hard edges as they denote a transition to a different soft edged group.
-	// Note that a vertex is usually only treated as 'hard' if it is met by two hard edges: if a hard edge terminates at the vertex we're looking at,
-	// it may be possible to find an adjacency path from one side to the other, and hence does not represent a barrier between soft edged groups.
-	// The exception is if the mesh is not closed and the hard edge leads to the mesh boundary.
-	static TArray<FEdgeID> ConnectedSoftEdges;
-	GetConnectedSoftEdges( VertexID, ConnectedSoftEdges );
-
-	// Iterate through adjacent polygons which contain the given vertex, but don't cross a hard edge.
-	// Maintain a list of polygon IDs to be examined. Adjacents are added to the list if suitable.
-	// Add the start poly here.
-	static TArray<FPolygonID> PolygonsToCheck;
-	PolygonsToCheck.Reset();
-	PolygonsToCheck.Add( PolygonID );
-
-	int32 Index = 0;
-	while( Index < PolygonsToCheck.Num() )
-	{
-		const FPolygonID PolygonToCheck = PolygonsToCheck[ Index ];
-		Index++;
-
-		const FVertexInstanceID VertexInstanceID = GetVertexInstanceInPolygonForVertex( PolygonToCheck, VertexID );
-		if( VertexInstanceID != FVertexInstanceID::Invalid || ( bPolygonNotYetInitialized && Index == 0 ) )
-		{
-			// If we got here, either:
-			//  a) the polygon contains the vertex instance referencing the vertex we're looking at, or
-			//  b) the polygon is the one we passed in, which may not yet have its vertex instances assigned, and we've specified bPolygonNotYetInitialized=true to allow this case
-
-			// If this polygon contains an instance of this vertex, add it to the result.
-			if( VertexInstanceID != FVertexInstanceID::Invalid )
-			{
-				OutVertexInstanceIDs.AddUnique( VertexInstanceID );
-			}
-
-			// Now look at its adjacent polygons. If they are joined by a soft edge which includes the vertex we're interested in, we want to consider them.
-			// We take a shortcut by doing this process in reverse: we already know all the soft edges we are interested in, so check if any of them
-			// have the current polygon as an adjacent.
-			for( const FEdgeID ConnectedSoftEdge : ConnectedSoftEdges )
-			{
-				const FMeshEdge& Edge = Edges[ ConnectedSoftEdge ];
-				if( Edge.ConnectedPolygons.Contains( PolygonToCheck ) )
-				{
-					for( const FPolygonID AdjacentPolygon : Edge.ConnectedPolygons )
-					{
-						// Only add new polygons which haven't yet been added to the list. This prevents circular runs of polygons triggering infinite loops.
-						PolygonsToCheck.AddUnique( AdjacentPolygon );
-					}
-				}
-			}
-		}
-	}
-}
-
-
 void UEditableMesh::SetPolygonsVertexAttributes( const TArray<FVertexAttributesForPolygon>& VertexAttributesForPolygons )
 {
 	UE_LOG( LogEditableMesh, Verbose, TEXT( "SetPolygonsVertexAttributes: %s" ), *LogHelpers::ArrayToString( VertexAttributesForPolygons ) );
diff --git a/Engine/Plugins/Runtime/EditableMesh/Source/EditableMesh/Public/EditableMesh.h b/Engine/Plugins/Runtime/EditableMesh/Source/EditableMesh/Public/EditableMesh.h
--- a/Engine/Plugins/Runtime/EditableMesh/Source/EditableMesh/Public/EditableMesh.h
+++ b/Engine/Plugins/Runtime/EditableMesh/Source/EditableMesh/Public/EditableMesh.h
@@ -15,7 +15,11 @@
 
 class UEditableMeshAdapter;
 
+#ifdef EDITABLE_MESH_DEBUG_LOG
 EDITABLEMESH_API DECLARE_LOG_CATEGORY_EXTERN( LogEditableMesh, All, All );
+#else
+EDITABLEMESH_API DECLARE_LOG_CATEGORY_EXTERN( LogEditableMesh, Log, All );
+#endif
 
 
 // @todo mesheditor: Comment these classes and enums!
@@ -218,13 +222,11 @@ protected:
 	void SetEdgeAttribute( const FEdgeID EdgeID, const FMeshElementAttributeData& Attribute );
 	void SetPolygonAttribute( const FPolygonID PolygonID, const FMeshElementAttributeData& Attribute );
 	void SetPolygonGroupAttribute( const FPolygonGroupID PolygonGroupID, const FMeshElementAttributeData& Attribute );
-	void SetEdgeHardness( const FEdgeID EdgeID, const bool bIsHard, const bool bIsUndo );
 	FVertexInstanceID CreateVertexInstanceForContourVertex( const FVertexAndAttributes& ContourVertex, const FPolygonID PolygonID );
 	void CreatePolygonContour( const TArray<FVertexAndAttributes>& Contour, TArray<FEdgeID>& OutExistingEdgeIDs, TArray<FEdgeID>& OutEdgeIDs, TArray<UMeshDescription::FContourPoint>& OutContourPoints );
 	void BackupPolygonContour( const FMeshPolygonContour& Contour, TArray<FVertexAndAttributes>& OutVerticesAndAttributes );
 	void GetConnectedSoftEdges( const FVertexID VertexID, TArray<FEdgeID>& OutConnectedSoftEdges ) const;
 	void GetVertexConnectedPolygonsInSameSoftEdgedGroup( const FVertexID VertexInstanceID, const FPolygonID PolygonID, TArray<FPolygonID>& OutPolygonIDs ) const;
-	void GetVertexInstancesInSameSoftEdgedGroup( const FVertexID VertexID, const FPolygonID PolygonID, const bool bPolygonNotYetInitialized, TArray<FVertexInstanceID>& OutVertexInstanceIDs ) const;
 	FVertexInstanceID GetVertexInstanceInPolygonForVertex( const FPolygonID PolygonID, const FVertexID VertexID ) const;
 	void SetPolygonContourVertexAttributes( FMeshPolygonContour& Contour, const FPolygonID PolygonID, const int32 HoleIndex, const TArray<FMeshElementAttributeList>& AttributeLists );
 	void SplitVertexInstanceInPolygons( const FVertexInstanceID VertexInstanceID, const TArray<FPolygonID>& PolygonIDs );
@@ -493,6 +495,9 @@ private:
 	/** Called during end modification to retriangulate polygons in the pending polygon list */
 	void RetriangulatePolygons();
 
+	/** Called during end modification to merge vertex instances whose mergeable attributes are all equal */
+	void MergeVertexInstances();
+
 	/** Tries to incrementally update the octree based off geometry that has changed since last time.  If that's not
 	    reasonable to do, then this will rebuild the entire octree from scratch */
 	void UpdateOrRebuildOctree();
@@ -542,6 +547,9 @@ public:
 	/** List of polygons requiring retriangulation */
 	TSet<FPolygonID> PolygonsPendingTriangulation;
 
+	/** List of candidate vertices for merging instances */
+	TSet<FVertexID> VerticesPendingMerging;
+
 	/** True if StartModification() has been called.  Call EndModification() when you've finished changing the mesh. */
 	bool bIsBeingModified;
 
