API improvements. Sanity check inputs to UEditableMesh methods. Access mesh elements solely by IDs instead of looking at the implementation details in the FMesh* classes.

diff --git a/Engine/Plugins/Runtime/EditableMesh/Source/EditableMesh/EditableMesh.cpp b/Engine/Plugins/Runtime/EditableMesh/Source/EditableMesh/EditableMesh.cpp
--- a/Engine/Plugins/Runtime/EditableMesh/Source/EditableMesh/EditableMesh.cpp
+++ b/Engine/Plugins/Runtime/EditableMesh/Source/EditableMesh/EditableMesh.cpp
@@ -254,7 +254,6 @@ static void BackupAllAttributes( FMeshElementAttributeList& AttributeList, const
 				const auto& Value = AttributeIndicesArray.GetArrayForIndex( Index )[ ElementID ];
 
 				AttributeList.Attributes.Emplace( AttributeName, Index, FMeshElementAttributeValue( Value ) );
-//				UE_LOG( LogTemp, Log, TEXT( "Name: %s, Index: %d, Value: %s" ), *AttributeName.ToString(), Index, *FMeshElementAttributeValue( Value ).ToString() );
 			}
 		}
 	);
@@ -275,7 +274,6 @@ static void BackupAttributesByPredicate( FMeshElementAttributeList& AttributeLis
 					const auto& Value = AttributeIndicesArray.GetArrayForIndex( Index )[ ElementID ];
 
 					AttributeList.Attributes.Emplace( AttributeName, Index, FMeshElementAttributeValue( Value ) );
-//					UE_LOG( LogTemp, Log, TEXT( "Name: %s, Index: %d, Value: %s" ), *AttributeName.ToString(), Index, *FMeshElementAttributeValue( Value ).ToString() );
 				}
 			}
 		}
@@ -825,34 +823,26 @@ int32 UEditableMesh::GetVertexCount() const
 
 bool UEditableMesh::IsValidVertex( const FVertexID VertexID ) const
 {
-	return GetMeshDescription()->Vertices().IsValid( VertexID );
+	return GetMeshDescription()->IsVertexValid( VertexID );
 }
 
 
 bool UEditableMesh::IsOrphanedVertex( const FVertexID VertexID ) const
 {
-	const FVertexInstanceArray& VertexInstances = GetMeshDescription()->VertexInstances();
-	for( const FVertexInstanceID VertexInstanceID : GetMeshDescription()->GetVertex( VertexID ).VertexInstanceIDs )
-	{
-		if( VertexInstances[ VertexInstanceID ].ConnectedPolygons.Num() > 0 )
-		{
-			return false;
-		}
-	}
-
-	return true;
+	return GetMeshDescription()->IsVertexOrphaned( VertexID );
 }
 
 
 int32 UEditableMesh::GetVertexConnectedEdgeCount( const FVertexID VertexID ) const
 {
-	return GetMeshDescription()->GetVertex( VertexID ).ConnectedEdgeIDs.Num();
+	return GetMeshDescription()->GetVertexConnectedEdges( VertexID ).Num();
 }
 
 
 FEdgeID UEditableMesh::GetVertexConnectedEdge( const FVertexID VertexID, const int32 ConnectedEdgeNumber ) const
 {
-	return GetMeshDescription()->GetVertex( VertexID ).ConnectedEdgeIDs[ ConnectedEdgeNumber ];
+	const TArray<FEdgeID>& ConnectedEdgeIDs = GetMeshDescription()->GetVertexConnectedEdges( VertexID );
+	return ( ConnectedEdgeNumber >= 0 && ConnectedEdgeNumber < ConnectedEdgeIDs.Num() ) ? ConnectedEdgeIDs[ ConnectedEdgeNumber ] : FEdgeID::Invalid;
 }
 
 
@@ -864,19 +854,20 @@ int32 UEditableMesh::GetVertexInstanceCount() const
 
 FVertexID UEditableMesh::GetVertexInstanceVertex( const FVertexInstanceID VertexInstanceID ) const
 {
-	return GetMeshDescription()->GetVertexInstance( VertexInstanceID ).VertexID;
+	return GetMeshDescription()->GetVertexInstanceVertex( VertexInstanceID );
 }
 
 
 int32 UEditableMesh::GetVertexInstanceConnectedPolygonCount( const FVertexInstanceID VertexInstanceID ) const
 {
-	return GetMeshDescription()->GetVertexInstance( VertexInstanceID ).ConnectedPolygons.Num();
+	return GetMeshDescription()->GetVertexInstanceConnectedPolygons( VertexInstanceID ).Num();
 }
 
 
 FPolygonID UEditableMesh::GetVertexInstanceConnectedPolygon( const FVertexInstanceID VertexInstanceID, const int32 ConnectedPolygonNumber ) const
 {
-	return GetMeshDescription()->GetVertexInstance( VertexInstanceID ).ConnectedPolygons[ ConnectedPolygonNumber ];
+	const TArray<FPolygonID>& ConnectedPolygonIDs = GetMeshDescription()->GetVertexInstanceConnectedPolygons( VertexInstanceID );
+	return ( ConnectedPolygonNumber >= 0 && ConnectedPolygonNumber < ConnectedPolygonIDs.Num() ) ? ConnectedPolygonIDs[ ConnectedPolygonNumber ] : FPolygonID::Invalid;
 }
 
 
@@ -888,26 +879,27 @@ int32 UEditableMesh::GetEdgeCount() const
 
 bool UEditableMesh::IsValidEdge( const FEdgeID EdgeID ) const
 {
-	return GetMeshDescription()->Edges().IsValid( EdgeID );
+	return GetMeshDescription()->IsEdgeValid( EdgeID );
 }
 
 
 FVertexID UEditableMesh::GetEdgeVertex( const FEdgeID EdgeID, const int32 EdgeVertexNumber ) const
 {
 	checkSlow( EdgeVertexNumber >= 0 && EdgeVertexNumber < 2 );
-	return GetMeshDescription()->GetEdge( EdgeID ).VertexIDs[ EdgeVertexNumber ];
+	return GetMeshDescription()->GetEdgeVertex( EdgeID, EdgeVertexNumber );
 }
 
 
 int32 UEditableMesh::GetEdgeConnectedPolygonCount( const FEdgeID EdgeID ) const
 {
-	return GetMeshDescription()->GetEdge( EdgeID ).ConnectedPolygons.Num();
+	return GetMeshDescription()->GetEdgeConnectedPolygons( EdgeID ).Num();
 }
 
 
 FPolygonID UEditableMesh::GetEdgeConnectedPolygon( const FEdgeID EdgeID, const int32 ConnectedPolygonNumber ) const
 {
-	return GetMeshDescription()->GetEdge( EdgeID ).ConnectedPolygons[ ConnectedPolygonNumber ];
+	const TArray<FPolygonID>& ConnectedPolygonIDs = GetMeshDescription()->GetEdgeConnectedPolygons( EdgeID );
+	return ( ConnectedPolygonNumber >= 0 && ConnectedPolygonNumber < ConnectedPolygonIDs.Num() ) ? ConnectedPolygonIDs[ ConnectedPolygonNumber ] : FPolygonID::Invalid;
 }
 
 
@@ -919,19 +911,20 @@ int32 UEditableMesh::GetPolygonGroupCount() const
 
 bool UEditableMesh::IsValidPolygonGroup( const FPolygonGroupID PolygonGroupID ) const
 {
-	return GetMeshDescription()->PolygonGroups().IsValid( PolygonGroupID );
+	return GetMeshDescription()->IsPolygonGroupValid( PolygonGroupID );
 }
 
 
 int32 UEditableMesh::GetPolygonCountInGroup( const FPolygonGroupID PolygonGroupID ) const
 {
-	return GetMeshDescription()->GetPolygonGroup( PolygonGroupID ).Polygons.Num();
+	return GetMeshDescription()->GetPolygonGroupPolygons( PolygonGroupID ).Num();
 }
 
 
 FPolygonID UEditableMesh::GetPolygonInGroup( const FPolygonGroupID PolygonGroupID, const int32 PolygonNumber ) const
 {
-	return GetMeshDescription()->GetPolygonGroup( PolygonGroupID ).Polygons[ PolygonNumber ];
+	const TArray<FPolygonID>& PolygonIDs = GetMeshDescription()->GetPolygonGroupPolygons( PolygonGroupID );
+	return ( PolygonNumber >= 0 && PolygonNumber < PolygonIDs.Num() ) ? PolygonIDs[ PolygonNumber ] : FPolygonID::Invalid;
 }
 
 
@@ -943,67 +936,72 @@ int32 UEditableMesh::GetPolygonCount() const
 
 bool UEditableMesh::IsValidPolygon( const FPolygonID PolygonID ) const
 {
-	return GetMeshDescription()->Polygons().IsValid( PolygonID );
+	return GetMeshDescription()->IsPolygonValid( PolygonID );
 }
 
 
 FPolygonGroupID UEditableMesh::GetGroupForPolygon( const FPolygonID PolygonID ) const
 {
-	return GetMeshDescription()->GetPolygon( PolygonID ).PolygonGroupID;
+	return GetMeshDescription()->GetPolygonPolygonGroup( PolygonID );
 }
 
 
 int32 UEditableMesh::GetPolygonPerimeterVertexCount( const FPolygonID PolygonID ) const
 {
-	return GetMeshDescription()->GetPolygon( PolygonID ).PerimeterContour.VertexInstanceIDs.Num();
+	return GetMeshDescription()->GetPolygonPerimeterVertexInstances( PolygonID ).Num();
 }
 
 
-FVertexID UEditableMesh::GetPolygonPerimeterVertex( const FPolygonID PolygonID, const int32 PolygonVertexNumber ) const
+FVertexInstanceID UEditableMesh::GetPolygonPerimeterVertexInstance( const FPolygonID PolygonID, const int32 PolygonVertexNumber ) const
 {
-	return GetMeshDescription()->GetVertexInstance( GetPolygonPerimeterVertexInstance( PolygonID, PolygonVertexNumber ) ).VertexID;
+	const TArray<FVertexInstanceID>& VertexInstanceIDs = GetMeshDescription()->GetPolygonPerimeterVertexInstances( PolygonID );
+	return ( PolygonVertexNumber >= 0 && PolygonVertexNumber < VertexInstanceIDs.Num() ) ? VertexInstanceIDs[ PolygonVertexNumber ] : FVertexInstanceID::Invalid;
 }
 
 
-FVertexInstanceID UEditableMesh::GetPolygonPerimeterVertexInstance( const FPolygonID PolygonID, const int32 PolygonVertexNumber ) const
+FVertexID UEditableMesh::GetPolygonPerimeterVertex( const FPolygonID PolygonID, const int32 PolygonVertexNumber ) const
 {
-	return GetMeshDescription()->GetPolygon( PolygonID ).PerimeterContour.VertexInstanceIDs[ PolygonVertexNumber ];
+	const FVertexInstanceID VertexInstanceID = GetPolygonPerimeterVertexInstance( PolygonID, PolygonVertexNumber );
+	return ( VertexInstanceID != FVertexInstanceID::Invalid ) ? GetMeshDescription()->GetVertexInstanceVertex( VertexInstanceID ) : FVertexID::Invalid;
 }
 
 
 int32 UEditableMesh::GetPolygonHoleCount( const FPolygonID PolygonID ) const
 {
-	return GetMeshDescription()->GetPolygon( PolygonID ).HoleContours.Num();
+	return GetMeshDescription()->GetNumPolygonHoles( PolygonID );
 }
 
 
 int32 UEditableMesh::GetPolygonHoleVertexCount( const FPolygonID PolygonID, const int32 HoleNumber ) const
 {
-	return GetMeshDescription()->GetPolygon( PolygonID ).HoleContours[ HoleNumber ].VertexInstanceIDs.Num();
+	return GetMeshDescription()->GetPolygonHoleVertexInstances( PolygonID, HoleNumber ).Num();
 }
 
 
-FVertexID UEditableMesh::GetPolygonHoleVertex( const FPolygonID PolygonID, const int32 HoleNumber, const int32 PolygonVertexNumber ) const
+FVertexInstanceID UEditableMesh::GetPolygonHoleVertexInstance( const FPolygonID PolygonID, const int32 HoleNumber, const int32 PolygonVertexNumber ) const
 {
-	return GetMeshDescription()->GetVertexInstance( GetPolygonHoleVertexInstance( PolygonID, HoleNumber, PolygonVertexNumber ) ).VertexID;
+	const TArray<FVertexInstanceID>& VertexInstanceIDs = GetMeshDescription()->GetPolygonHoleVertexInstances( PolygonID, HoleNumber );
+	return ( PolygonVertexNumber >= 0 && PolygonVertexNumber < VertexInstanceIDs.Num() ) ? VertexInstanceIDs[ PolygonVertexNumber ] : FVertexInstanceID::Invalid;
 }
 
 
-FVertexInstanceID UEditableMesh::GetPolygonHoleVertexInstance( const FPolygonID PolygonID, const int32 HoleNumber, const int32 PolygonVertexNumber ) const
+FVertexID UEditableMesh::GetPolygonHoleVertex( const FPolygonID PolygonID, const int32 HoleNumber, const int32 PolygonVertexNumber ) const
 {
-	return GetMeshDescription()->GetPolygon( PolygonID ).HoleContours[ HoleNumber ].VertexInstanceIDs[ PolygonVertexNumber ];
+	const FVertexInstanceID VertexInstanceID = GetPolygonHoleVertexInstance( PolygonID, HoleNumber, PolygonVertexNumber );
+	return ( VertexInstanceID != FVertexInstanceID::Invalid ) ? GetMeshDescription()->GetVertexInstanceVertex( VertexInstanceID ) : FVertexID::Invalid;
 }
 
 
 int32 UEditableMesh::GetPolygonTriangulatedTriangleCount( const FPolygonID PolygonID ) const
 {
-	return GetMeshDescription()->GetPolygon( PolygonID ).Triangles.Num();
+	return GetMeshDescription()->GetPolygonTriangles( PolygonID ).Num();
 }
 
 
 FMeshTriangle UEditableMesh::GetPolygonTriangulatedTriangle( const FPolygonID PolygonID, int32 PolygonTriangleNumber ) const
 {
-	return GetMeshDescription()->GetPolygon( PolygonID ).Triangles[ PolygonTriangleNumber ];
+	const TArray<FMeshTriangle>& Triangles = GetMeshDescription()->GetPolygonTriangles( PolygonID );
+	return ( PolygonTriangleNumber >= 0 && PolygonTriangleNumber < Triangles.Num() ) ? Triangles[ PolygonTriangleNumber ] : FMeshTriangle();
 }
 
 
@@ -1187,7 +1185,6 @@ void UEditableMesh::SplitVertexInstanceInPolygons( const FVertexInstanceID Verte
 	CreateVertexInstances( VertexInstancesToCreate, NewVertexInstanceIDs );
 
 	ReplaceVertexInstanceInPolygons( VertexInstanceID, NewVertexInstanceIDs[ 0 ], PolygonIDs );
-	UE_LOG( LogTemp, Log, TEXT( "Split vertex instance %d: new instance %d" ), VertexInstanceID.GetValue(), NewVertexInstanceIDs[ 0 ].GetValue() );
 
 	UE_LOG( LogEditableMesh, Verbose, TEXT( "* SplitVertexInstanceInPolygons returned" ) );
 }
@@ -1373,7 +1370,6 @@ void UEditableMesh::MergeVertexInstances()
 					// Change occurrences of VertexInstanceB for VertexInstanceA in VertexInstanceB's connected polygons.
 					// Note, this will cause VertexInstanceA's connected polygons list to be added to (at the end).
 					// This works because we are evaluating the number of connected polygons each time round the loop.
-					UE_LOG( LogTemp, Log, TEXT( "Replacing vertex instance %d with vertex instance %d" ), VertexInstanceIDB.GetValue(), VertexInstanceIDA.GetValue() )
 					ReplaceVertexInstanceInPolygons( VertexInstanceIDB, VertexInstanceIDA, ConnectedPolygonIDsB );
 
 					// This will also cause VertexInstanceB to be disconnected from all polygons.
@@ -1417,85 +1413,44 @@ bool UEditableMesh::IsPreviewingSubdivisions() const
 
 void UEditableMesh::GetVertexConnectedEdges( const FVertexID VertexID, TArray<FEdgeID>& OutConnectedEdgeIDs ) const
 {
-	OutConnectedEdgeIDs.Reset();
-
-	const int32 ConnectedEdgeCount = GetVertexConnectedEdgeCount( VertexID );
-	OutConnectedEdgeIDs.Reserve( ConnectedEdgeCount );
-	for( int32 EdgeNumber = 0; EdgeNumber < ConnectedEdgeCount; ++EdgeNumber )
-	{
-		const FEdgeID ConnectedEdgeID = GetVertexConnectedEdge( VertexID, EdgeNumber );
-		OutConnectedEdgeIDs.Add( ConnectedEdgeID );
-	}
+	OutConnectedEdgeIDs = GetMeshDescription()->GetVertexConnectedEdges( VertexID );
 }
 
 
 void UEditableMesh::GetVertexConnectedPolygons( const FVertexID VertexID, TArray<FPolygonID>& OutConnectedPolygonIDs ) const
 {
-	OutConnectedPolygonIDs.Reset();
-
-	const FVertexInstanceArray& VertexInstances = GetMeshDescription()->VertexInstances();
-	for( const FVertexInstanceID VertexInstanceID : GetMeshDescription()->GetVertex( VertexID ).VertexInstanceIDs )
-	{
-		OutConnectedPolygonIDs.Append( VertexInstances[ VertexInstanceID ].ConnectedPolygons );
-	}
+	GetMeshDescription()->GetVertexConnectedPolygons( VertexID, OutConnectedPolygonIDs );
 }
 
 
 void UEditableMesh::GetVertexInstanceConnectedPolygons( const FVertexInstanceID VertexInstanceID, TArray<FPolygonID>& OutConnectedPolygonIDs ) const
 {
-	OutConnectedPolygonIDs = GetMeshDescription()->GetVertexInstance( VertexInstanceID ).ConnectedPolygons;
+	OutConnectedPolygonIDs = GetMeshDescription()->GetVertexInstanceConnectedPolygons( VertexInstanceID );
 }
 
 
 void UEditableMesh::GetVertexAdjacentVertices( const FVertexID VertexID, TArray<FVertexID>& OutAdjacentVertexIDs ) const
 {
-	OutAdjacentVertexIDs.Reset();
-
-	const FEdgeArray& Edges = GetMeshDescription()->Edges();
-	for( const FEdgeID EdgeID : GetMeshDescription()->GetVertex( VertexID ).ConnectedEdgeIDs )
-	{
-		const FMeshEdge& Edge = Edges[ EdgeID ];
-		OutAdjacentVertexIDs.Add( Edge.VertexIDs[ 0 ] == VertexID ? Edge.VertexIDs[ 1 ] : Edge.VertexIDs[ 0 ] );
-	}
+	GetMeshDescription()->GetVertexAdjacentVertices( VertexID, OutAdjacentVertexIDs );
 }
 
 
 void UEditableMesh::GetEdgeVertices( const FEdgeID EdgeID, FVertexID& OutEdgeVertexID0, FVertexID& OutEdgeVertexID1 ) const
 {
-	const FMeshEdge& Edge = GetMeshDescription()->GetEdge( EdgeID );
-	OutEdgeVertexID0 = Edge.VertexIDs[ 0 ];
-	OutEdgeVertexID1 = Edge.VertexIDs[ 1 ];
+	OutEdgeVertexID0 = GetMeshDescription()->GetEdgeVertex( EdgeID, 0 );
+	OutEdgeVertexID1 = GetMeshDescription()->GetEdgeVertex( EdgeID, 1 );
 }
 
 
 void UEditableMesh::GetEdgeConnectedPolygons( const FEdgeID EdgeID, TArray<FPolygonID>& OutConnectedPolygonIDs ) const
 {
-	OutConnectedPolygonIDs = GetMeshDescription()->GetEdge( EdgeID ).ConnectedPolygons;
+	OutConnectedPolygonIDs = GetMeshDescription()->GetEdgeConnectedPolygons( EdgeID );
 }
 
 
 FEdgeID UEditableMesh::GetEdgeThatConnectsVertices( const FVertexID VertexID0, const FVertexID VertexID1 ) const
 {
-	FEdgeID EdgeID = FEdgeID::Invalid;
-
-	const int32 Vertex0ConnectedEdgeCount = GetVertexConnectedEdgeCount( VertexID0 );
-	for( int32 Vertex0ConnectedEdgeNumber = 0; Vertex0ConnectedEdgeNumber < Vertex0ConnectedEdgeCount; ++Vertex0ConnectedEdgeNumber )
-	{
-		const FEdgeID Vertex0ConnectedEdge = GetVertexConnectedEdge( VertexID0, Vertex0ConnectedEdgeNumber );
-
-		FVertexID EdgeVertexID0, EdgeVertexID1;
-		GetEdgeVertices( Vertex0ConnectedEdge, /* Out */ EdgeVertexID0, /* Out */ EdgeVertexID1 );
-
-		if( ( EdgeVertexID0 == VertexID0 && EdgeVertexID1 == VertexID1 ) ||
-			( EdgeVertexID0 == VertexID1 && EdgeVertexID1 == VertexID0 ) )
-		{
-			EdgeID = Vertex0ConnectedEdge;
-			break;
-		}
-	}
-
-	check( EdgeID != FEdgeID::Invalid );
-	return EdgeID;
+	return GetMeshDescription()->GetVertexPairEdge( VertexID0, VertexID1 );
 }
 
 
@@ -1503,9 +1458,6 @@ void UEditableMesh::GetEdgeLoopElements( const FEdgeID EdgeID, TArray<FEdgeID>&
 {
 	EdgeLoopIDs.Reset();
 
-	const FEdgeArray& Edges = GetMeshDescription()->Edges();
-	const FVertexArray& Vertices = GetMeshDescription()->Vertices();
-
 	// Maintain a list of unique edge IDs which form the loop
 	static TSet<FEdgeID> EdgeIDs;
 	EdgeIDs.Reset();
@@ -1524,14 +1476,16 @@ void UEditableMesh::GetEdgeLoopElements( const FEdgeID EdgeID, TArray<FEdgeID>&
 		const FVertexID FromVertexID = CurrentItem.Get<1>();
 		EdgeIDs.Add( CurrentEdgeID );
 
-		// Get the edge so we can see which polygons are connected to it.
+		// See which polygons are connected to the current edge.
 		// When continuing the loop, the criterion is that new edges must share no polygons with this edge,
 		// i.e. they are the other side of a perpendicular edge.
-		const FMeshEdge& Edge = Edges[ CurrentEdgeID ];
+		const TArray<FPolygonID>& CurrentEdgeConnectedPolygons = GetMeshDescription()->GetEdgeConnectedPolygons( CurrentEdgeID );
 
 		// Now look for edges connected to each end of this edge
-		for( const FVertexID ConnectedVertexID : Edge.VertexIDs )
+		for( int32 EdgeVertexIndex = 0; EdgeVertexIndex < 2; ++EdgeVertexIndex )
 		{
+			const FVertexID ConnectedVertexID = GetMeshDescription()->GetEdgeVertex( CurrentEdgeID, EdgeVertexIndex );
+
 			// If we have already processed this vertex, skip it
 			if( ConnectedVertexID == FromVertexID )
 			{
@@ -1542,18 +1496,16 @@ void UEditableMesh::GetEdgeLoopElements( const FEdgeID EdgeID, TArray<FEdgeID>&
 			FEdgeID AdjacentEdgeID = FEdgeID::Invalid;
 
 			// Iterate through all edges connected to this vertex
-			const FMeshVertex& Vertex = Vertices[ ConnectedVertexID ];
-			for( const FEdgeID ConnectedEdgeID : Vertex.ConnectedEdgeIDs )
+			for( const FEdgeID ConnectedEdgeID : GetMeshDescription()->GetVertexConnectedEdges( ConnectedVertexID ) )
 			{
 				// If this edge hasn't been added to the loop...
 				if( !EdgeIDs.Contains( ConnectedEdgeID ) )
 				{
 					// ...see if it shares any polygons with the original edge (intersection operation)
 					bool bIsCandidateEdge = true;
-					const FMeshEdge& ConnectedEdge = Edges[ ConnectedEdgeID ];
-					for( const FPolygonID ConnectedPolygonID : ConnectedEdge.ConnectedPolygons )
+					for( const FPolygonID ConnectedPolygonID : GetMeshDescription()->GetEdgeConnectedPolygons( ConnectedEdgeID ) )
 					{
-						if( Edge.ConnectedPolygons.Contains( ConnectedPolygonID ) )
+						if( CurrentEdgeConnectedPolygons.Contains( ConnectedPolygonID ) )
 						{
 							bIsCandidateEdge = false;
 							break;
@@ -1608,138 +1560,121 @@ int32 UEditableMesh::GetPolygonHoleEdgeCount( const FPolygonID PolygonID, const
 
 void UEditableMesh::GetPolygonPerimeterVertices( const FPolygonID PolygonID, TArray<FVertexID>& OutPolygonPerimeterVertexIDs ) const
 {
-	const FVertexInstanceArray& VertexInstances = GetMeshDescription()->VertexInstances();
-	const FMeshPolygon& Polygon = GetMeshDescription()->GetPolygon( PolygonID );
-
-	OutPolygonPerimeterVertexIDs.SetNumUninitialized( Polygon.PerimeterContour.VertexInstanceIDs.Num(), false );
-
-	int32 Index = 0;
-	for( const FVertexInstanceID VertexInstanceID : Polygon.PerimeterContour.VertexInstanceIDs )
-	{
-		const FMeshVertexInstance& VertexInstance = VertexInstances[ VertexInstanceID ];
-		OutPolygonPerimeterVertexIDs[ Index ] = VertexInstance.VertexID;
-		Index++;
-	}
+	GetMeshDescription()->GetPolygonPerimeterVertices( PolygonID, OutPolygonPerimeterVertexIDs );
 }
 
 
 void UEditableMesh::GetPolygonPerimeterVertexInstances( const FPolygonID PolygonID, TArray<FVertexInstanceID>& OutPolygonPerimeterVertexInstanceIDs ) const
 {
-	const FMeshPolygon& Polygon = GetMeshDescription()->GetPolygon( PolygonID );
-	OutPolygonPerimeterVertexInstanceIDs = Polygon.PerimeterContour.VertexInstanceIDs;
+	OutPolygonPerimeterVertexInstanceIDs = GetMeshDescription()->GetPolygonPerimeterVertexInstances( PolygonID );
 }
 
 
 void UEditableMesh::GetPolygonHoleVertices( const FPolygonID PolygonID, const int32 HoleNumber, TArray<FVertexID>& OutHoleVertexIDs ) const
 {
-	const FVertexInstanceArray& VertexInstances = GetMeshDescription()->VertexInstances();
-	const FMeshPolygon& Polygon = GetMeshDescription()->GetPolygon( PolygonID );
-	const FMeshPolygonContour& Contour = Polygon.HoleContours[ HoleNumber ];
-
-	OutHoleVertexIDs.SetNumUninitialized( Contour.VertexInstanceIDs.Num(), false );
-
-	int32 Index = 0;
-	for( const FVertexInstanceID VertexInstanceID : Contour.VertexInstanceIDs )
+	OutHoleVertexIDs.Reset();	// in case it fails below
+	if( HoleNumber < GetMeshDescription()->GetNumPolygonHoles( PolygonID ) )
 	{
-		const FMeshVertexInstance& VertexInstance = VertexInstances[ VertexInstanceID ];
-		OutHoleVertexIDs[ Index ] = VertexInstance.VertexID;
-		Index++;
+		GetMeshDescription()->GetPolygonHoleVertices( PolygonID, HoleNumber, OutHoleVertexIDs );
 	}
 }
 
 
 void UEditableMesh::GetPolygonHoleVertexInstances( const FPolygonID PolygonID, const int32 HoleNumber, TArray<FVertexInstanceID>& OutHoleVertexInstanceIDs ) const
 {
-	const FMeshPolygon& Polygon = GetMeshDescription()->GetPolygon( PolygonID );
-	const FMeshPolygonContour& Contour = Polygon.HoleContours[ HoleNumber ];
-
-	OutHoleVertexInstanceIDs = Contour.VertexInstanceIDs;
-}
-
-
-FEdgeID UEditableMesh::GetPolygonContourEdge( const FMeshPolygonContour& Contour, const int32 ContourEdgeNumber, bool& bOutEdgeWindingIsReversed ) const
-{
-	const int32 NumPolygonContourEdges = Contour.VertexInstanceIDs.Num();
-
-	const FVertexInstanceID VertexInstanceID = Contour.VertexInstanceIDs[ ContourEdgeNumber ];
-	const FVertexInstanceID NextVertexInstanceID = Contour.VertexInstanceIDs[ ( ContourEdgeNumber + 1 ) % NumPolygonContourEdges ];
-
-	const FVertexInstanceArray& VertexInstances = GetMeshDescription()->VertexInstances();
-	const FVertexID VertexID = VertexInstances[ VertexInstanceID ].VertexID;
-	const FVertexID NextVertexID = VertexInstances[ NextVertexInstanceID ].VertexID;
-
-	return GetVertexPairEdge( VertexID, NextVertexID, bOutEdgeWindingIsReversed );
+	OutHoleVertexInstanceIDs.Reset();	// in case it fails below
+	if( HoleNumber < GetMeshDescription()->GetNumPolygonHoles( PolygonID ) )
+	{
+		OutHoleVertexInstanceIDs = GetMeshDescription()->GetPolygonHoleVertexInstances( PolygonID, HoleNumber );
+	}
 }
 
 
 FEdgeID UEditableMesh::GetPolygonPerimeterEdge( const FPolygonID PolygonID, const int32 PerimeterEdgeNumber, bool& bOutEdgeWindingIsReversedForPolygon ) const
 {
-	return GetPolygonContourEdge( GetMeshDescription()->GetPolygon( PolygonID ).PerimeterContour, PerimeterEdgeNumber, bOutEdgeWindingIsReversedForPolygon );
+	const UMeshDescription* Description = GetMeshDescription();
+	const TArray<FVertexInstanceID>& VertexInstanceIDs = Description->GetPolygonPerimeterVertexInstances( PolygonID );
+	if( PerimeterEdgeNumber >= VertexInstanceIDs.Num() )
+	{
+		return FEdgeID::Invalid;
+	}
+
+	const FVertexID VertexID0 = Description->GetVertexInstanceVertex( VertexInstanceIDs[ PerimeterEdgeNumber ] );
+	const FVertexID VertexID1 = Description->GetVertexInstanceVertex( VertexInstanceIDs[ ( PerimeterEdgeNumber + 1 ) % VertexInstanceIDs.Num() ] );
+	return GetVertexPairEdge( VertexID0, VertexID1, bOutEdgeWindingIsReversedForPolygon );
 }
 
 
 FEdgeID UEditableMesh::GetVertexPairEdge( const FVertexID StartVertexID, const FVertexID EndVertexID, bool& bOutEdgeWindingIsReversed ) const
 {
-	bOutEdgeWindingIsReversed = false;
-
-	const FEdgeArray& Edges = GetMeshDescription()->Edges();
-
-	for( const FEdgeID VertexConnectedEdgeID : GetMeshDescription()->GetVertex( StartVertexID ).ConnectedEdgeIDs )
+	const FEdgeID EdgeID = GetMeshDescription()->GetVertexPairEdge( StartVertexID, EndVertexID );
+	if( EdgeID != FEdgeID::Invalid )
 	{
-		const FMeshEdge& Edge = Edges[ VertexConnectedEdgeID ];
-
-		// Try the edge's first vertex.  Does it point to us?
-		FVertexID OtherEdgeVertexID = Edge.VertexIDs[ 0 ];
-		bOutEdgeWindingIsReversed = false;
-		if( OtherEdgeVertexID == StartVertexID )
-		{
-			// Must be the the other one
-			OtherEdgeVertexID = Edge.VertexIDs[ 1 ];
-			bOutEdgeWindingIsReversed = true;
-		}
-
-		if( OtherEdgeVertexID == EndVertexID )
-		{
-			// We found the edge!
-			return VertexConnectedEdgeID;
-		}
+		// @note: seems like a reversal of logic here, but everything else expects this behavior now
+		bOutEdgeWindingIsReversed = ( GetMeshDescription()->GetEdgeVertex( EdgeID, 0 ) == StartVertexID );
 	}
 
-	return FEdgeID::Invalid;
+	return EdgeID;
 }
 
 
 FEdgeID UEditableMesh::GetPolygonHoleEdge( const FPolygonID PolygonID, const int32 HoleNumber, const int32 HoleEdgeNumber ) const
 {
-	bool bOutEdgeWindingIsReversed;
-	return GetPolygonContourEdge( GetMeshDescription()->GetPolygon( PolygonID ).HoleContours[ HoleEdgeNumber ], HoleEdgeNumber, bOutEdgeWindingIsReversed );
-}
+	const UMeshDescription* Description = GetMeshDescription();
 
+	if( HoleNumber >= Description->GetNumPolygonHoles( PolygonID ) )
+	{
+		return FEdgeID::Invalid;
+	}
 
-void UEditableMesh::GetPolygonContourEdges( const FMeshPolygonContour& Contour, TArray<FEdgeID>& OutPolygonContourEdgeIDs ) const
-{
-	const int32 NumContourEdges = Contour.VertexInstanceIDs.Num();
-	OutPolygonContourEdgeIDs.SetNumUninitialized( NumContourEdges, false );
-
-	for( int32 EdgeNumber = 0; EdgeNumber < NumContourEdges; ++EdgeNumber )
+	const TArray<FVertexInstanceID>& VertexInstanceIDs = Description->GetPolygonHoleVertexInstances( PolygonID, HoleNumber );
+	if( HoleEdgeNumber >= VertexInstanceIDs.Num() )
 	{
-		bool bOutEdgeWindingIsReversed;
-		OutPolygonContourEdgeIDs[ EdgeNumber ] = GetPolygonContourEdge( Contour, EdgeNumber, bOutEdgeWindingIsReversed );
+		return FEdgeID::Invalid;
 	}
+
+	const FVertexID VertexID0 = Description->GetVertexInstanceVertex( VertexInstanceIDs[ HoleEdgeNumber ] );
+	const FVertexID VertexID1 = Description->GetVertexInstanceVertex( VertexInstanceIDs[ ( HoleEdgeNumber + 1 ) % VertexInstanceIDs.Num() ] );
+	return Description->GetVertexPairEdge( VertexID0, VertexID1 );
 }
 
 
 void UEditableMesh::GetPolygonPerimeterEdges( const FPolygonID PolygonID, TArray<FEdgeID>& OutPolygonPerimeterEdgeIDs ) const
 {
-	GetPolygonContourEdges( GetMeshDescription()->GetPolygon( PolygonID ).PerimeterContour, OutPolygonPerimeterEdgeIDs );
+	const UMeshDescription* Description = GetMeshDescription();
+	const TArray<FVertexInstanceID>& VertexInstanceIDs = Description->GetPolygonPerimeterVertexInstances( PolygonID );
+	const int32 NumContourEdges = VertexInstanceIDs.Num();
+
+	OutPolygonPerimeterEdgeIDs.SetNumUninitialized( NumContourEdges, false );
+	for( int32 Index = 0; Index < NumContourEdges; ++Index )
+	{
+		const int32 NextIndex = ( Index + 1 == NumContourEdges ) ? 0 : Index + 1;
+		const FVertexID VertexID0 = Description->GetVertexInstanceVertex( VertexInstanceIDs[ Index ] );
+		const FVertexID VertexID1 = Description->GetVertexInstanceVertex( VertexInstanceIDs[ NextIndex ] );
+		OutPolygonPerimeterEdgeIDs[ Index ] = Description->GetVertexPairEdge( VertexID0, VertexID1 );
+	}
 }
 
 
 void UEditableMesh::GetPolygonHoleEdges( const FPolygonID PolygonID, const int32 HoleNumber, TArray<FEdgeID>& OutHoleEdgeIDs ) const
 {
-	const FMeshPolygon& Polygon = GetMeshDescription()->GetPolygon( PolygonID );
-	const FMeshPolygonContour& Contour = Polygon.HoleContours[ HoleNumber ];
-	GetPolygonContourEdges( Contour, OutHoleEdgeIDs );
+	const UMeshDescription* Description = GetMeshDescription();
+
+	OutHoleEdgeIDs.Reset();
+	if( HoleNumber < Description->GetNumPolygonHoles( PolygonID ) )
+	{
+		const TArray<FVertexInstanceID>& VertexInstanceIDs = Description->GetPolygonHoleVertexInstances( PolygonID, HoleNumber );
+		const int32 NumContourEdges = VertexInstanceIDs.Num();
+
+		OutHoleEdgeIDs.SetNumUninitialized( NumContourEdges, false );
+		for( int32 Index = 0; Index < NumContourEdges; ++Index )
+		{
+			const int32 NextIndex = ( Index + 1 == NumContourEdges ) ? 0 : Index + 1;
+			const FVertexID VertexID0 = Description->GetVertexInstanceVertex( VertexInstanceIDs[ Index ] );
+			const FVertexID VertexID1 = Description->GetVertexInstanceVertex( VertexInstanceIDs[ NextIndex ] );
+			OutHoleEdgeIDs[ Index ] = Description->GetVertexPairEdge( VertexID0, VertexID1 );
+		}
+	}
 }
 
 
diff --git a/Engine/Plugins/Runtime/EditableMesh/Source/EditableMesh/Public/EditableMesh.h b/Engine/Plugins/Runtime/EditableMesh/Source/EditableMesh/Public/EditableMesh.h
--- a/Engine/Plugins/Runtime/EditableMesh/Source/EditableMesh/Public/EditableMesh.h
+++ b/Engine/Plugins/Runtime/EditableMesh/Source/EditableMesh/Public/EditableMesh.h
@@ -217,8 +217,6 @@ protected:
 	/** Given a set of index remappings, fixes up references in the octree */
 	void RemapOctreeIDs( const FElementIDRemappings& Remappings );
 
-	FEdgeID GetPolygonContourEdge( const FMeshPolygonContour& Contour, const int32 ContourEdgeNumber, bool& bOutEdgeWindingIsReversedForPolygon ) const;
-	void GetPolygonContourEdges( const FMeshPolygonContour& Contour, TArray<FEdgeID>& OutPolygonContourEdgeIDs ) const;
 	void SetVertexAttribute( const FVertexID VertexID, const FMeshElementAttributeData& Attribute );
 	void SetVertexInstanceAttribute( const FVertexInstanceID VertexInstanceID, const FMeshElementAttributeData& Attribute );
 	void SetEdgeAttribute( const FEdgeID EdgeID, const FMeshElementAttributeData& Attribute );
diff --git a/Engine/Source/Runtime/MeshDescription/Private/MeshDescription.cpp b/Engine/Source/Runtime/MeshDescription/Private/MeshDescription.cpp
--- a/Engine/Source/Runtime/MeshDescription/Private/MeshDescription.cpp
+++ b/Engine/Source/Runtime/MeshDescription/Private/MeshDescription.cpp
@@ -393,21 +393,53 @@ void UMeshDescription::RemapAttributes( const FElementIDRemappings& Remappings )
 	PolygonGroupAttributesSet.Remap( Remappings.NewPolygonGroupIndexLookup );
 }
 
-void UMeshDescription::GetPolygonPerimeterVertices(const FPolygonID PolygonID, TArray<FVertexID>& OutPolygonPerimeterVertexIDs) const
+
+bool UMeshDescription::IsVertexOrphaned( const FVertexID VertexID ) const
 {
-	const FMeshPolygon& Polygon = GetPolygon(PolygonID);
+	for( const FVertexInstanceID VertexInstanceID : GetVertex( VertexID ).VertexInstanceIDs )
+	{
+		if( GetVertexInstance( VertexInstanceID ).ConnectedPolygons.Num() > 0 )
+		{
+			return false;
+		}
+	}
 
-	OutPolygonPerimeterVertexIDs.SetNumUninitialized(Polygon.PerimeterContour.VertexInstanceIDs.Num(), false);
+	return true;
+}
+
+
+void UMeshDescription::GetPolygonPerimeterVertices( const FPolygonID PolygonID, TArray<FVertexID>& OutPolygonPerimeterVertexIDs ) const
+{
+	const FMeshPolygon& Polygon = GetPolygon( PolygonID );
+
+	OutPolygonPerimeterVertexIDs.SetNumUninitialized( Polygon.PerimeterContour.VertexInstanceIDs.Num(), false );
 
 	int32 Index = 0;
-	for (const FVertexInstanceID VertexInstanceID : Polygon.PerimeterContour.VertexInstanceIDs)
+	for( const FVertexInstanceID VertexInstanceID : Polygon.PerimeterContour.VertexInstanceIDs )
 	{
-		const FMeshVertexInstance& VertexInstance = VertexInstanceArray[VertexInstanceID];
-		OutPolygonPerimeterVertexIDs[Index] = VertexInstance.VertexID;
+		const FMeshVertexInstance& VertexInstance = VertexInstanceArray[ VertexInstanceID ];
+		OutPolygonPerimeterVertexIDs[ Index ] = VertexInstance.VertexID;
 		Index++;
 	}
 }
 
+
+void UMeshDescription::GetPolygonHoleVertices( const FPolygonID PolygonID, const int32 HoleIndex, TArray<FVertexID>& OutPolygonHoleVertexIDs ) const
+{
+	const FMeshPolygon& Polygon = GetPolygon( PolygonID );
+
+	OutPolygonHoleVertexIDs.SetNumUninitialized( Polygon.HoleContours[ HoleIndex ].VertexInstanceIDs.Num(), false );
+
+	int32 Index = 0;
+	for( const FVertexInstanceID VertexInstanceID : Polygon.HoleContours[ HoleIndex ].VertexInstanceIDs )
+	{
+		const FMeshVertexInstance& VertexInstance = VertexInstanceArray[ VertexInstanceID ];
+		OutPolygonHoleVertexIDs[ Index ] = VertexInstance.VertexID;
+		Index++;
+	}
+}
+
+
 /** Given three direction vectors, indicates if A and B are on the same 'side' of Vec. */
 bool UMeshDescription::VectorsOnSameSide(const FVector& Vec, const FVector& A, const FVector& B, const float SameSideDotProductEpsilon)
 {
diff --git a/Engine/Source/Runtime/MeshDescription/Public/MeshDescription.h b/Engine/Source/Runtime/MeshDescription/Public/MeshDescription.h
--- a/Engine/Source/Runtime/MeshDescription/Public/MeshDescription.h
+++ b/Engine/Source/Runtime/MeshDescription/Public/MeshDescription.h
@@ -614,7 +614,12 @@ public:
 
 //////////////////////////////////////////////////////////////////////////
 // Meshdescription general functions
+
 public:
+
+	/** Returns whether a given vertex is orphaned, i.e. it doesn't form part of any polygon */
+	bool IsVertexOrphaned( const FVertexID VertexID ) const;
+
 	/** Returns the edge ID defined by the two given vertex IDs, if there is one; otherwise FEdgeID::Invalid */
 	FEdgeID GetVertexPairEdge( const FVertexID VertexID0, const FVertexID VertexID1 ) const
 	{
@@ -643,6 +648,31 @@ public:
 		return VertexArray[ VertexID ].VertexInstanceIDs;
 	}
 
+	/** Populates the passed array of PolygonIDs with the polygons connected to this vertex */
+	void GetVertexConnectedPolygons( const FVertexID VertexID, TArray<FPolygonID>& OutConnectedPolygonIDs ) const
+	{
+		OutConnectedPolygonIDs.Reset();
+		for( const FVertexInstanceID VertexInstanceID : VertexArray[ VertexID ].VertexInstanceIDs )
+		{
+			OutConnectedPolygonIDs.Append( VertexInstanceArray[ VertexInstanceID ].ConnectedPolygons );
+		}
+	}
+
+	/** Populates the passed array of VertexIDs with the vertices adjacent to this vertex */
+	void GetVertexAdjacentVertices( const FVertexID VertexID, TArray<FVertexID>& OutAdjacentVertexIDs ) const
+	{
+		const TArray<FEdgeID>& ConnectedEdgeIDs = VertexArray[ VertexID ].ConnectedEdgeIDs;
+		OutAdjacentVertexIDs.SetNumUninitialized( ConnectedEdgeIDs.Num() );
+
+		int32 Index = 0;
+		for( const FEdgeID EdgeID : ConnectedEdgeIDs )
+		{
+			const FMeshEdge& Edge = EdgeArray[ EdgeID ];
+			OutAdjacentVertexIDs[ Index ] = ( Edge.VertexIDs[ 0 ] == VertexID ) ? Edge.VertexIDs[ 1 ] : Edge.VertexIDs[ 0 ];
+			Index++;
+		}
+	}
+
 	/** Returns reference to an array of polygon IDs connected to this edge */
 	const TArray<FPolygonID>& GetEdgeConnectedPolygons( const FEdgeID EdgeID ) const
 	{
@@ -680,6 +710,24 @@ public:
 		return PolygonArray[ PolygonID ].PerimeterContour.VertexInstanceIDs;
 	}
 
+	/** Populates the passed array of VertexIDs with the vertices which form the polygon perimeter */
+	void GetPolygonPerimeterVertices( const FPolygonID PolygonID, TArray<FVertexID>& OutPolygonPerimeterVertexIDs ) const;
+
+	/** Returns the number of holes in this polygon */
+	int32 GetNumPolygonHoles( const FPolygonID PolygonID ) const
+	{
+		return PolygonArray[ PolygonID ].HoleContours.Num();
+	}
+
+	/** Populates the passed array of VertexIDs with the vertices which form the contour of the specified hole */
+	void GetPolygonHoleVertices( const FPolygonID PolygonID, const int32 HoleIndex, TArray<FVertexID>& OutPolygonHoleVertexIDs ) const;
+
+	/** Returns reference to an array of VertexInstance IDs forming the contour of the given hole */
+	const TArray<FVertexInstanceID>& GetPolygonHoleVertexInstances( const FPolygonID PolygonID, const int32 HoleIndex ) const
+	{
+		return PolygonArray[ PolygonID ].HoleContours[ HoleIndex ].VertexInstanceIDs;
+	}
+
 	void GetPolygonEdges(const FPolygonID PolygonID, TArray<FEdgeID>& OutPolygonEdges) const
 	{
 		const FMeshPolygonContour& PerimeterContour = PolygonArray[PolygonID].PerimeterContour;
@@ -756,7 +804,6 @@ public:
 private:
 	bool VectorsOnSameSide(const FVector& Vec, const FVector& A, const FVector& B, const float SameSideDotProductEpsilon);
 	bool PointInTriangle(const FVector& A, const FVector& B, const FVector& C, const FVector& P, const float InsideTriangleDotProductEpsilon);
-	void GetPolygonPerimeterVertices(const FPolygonID PolygonID, TArray<FVertexID>& OutPolygonPerimeterVertexIDs) const;
 	FPlane ComputePolygonPlane(const FPolygonID PolygonID) const;
 	FVector ComputePolygonNormal(const FPolygonID PolygonID) const;
 
