More fixes for Linux build.

diff --git a/Engine/Plugins/Editor/MeshEditor/Source/PolygonModeling/RemoveEdge.cpp b/Engine/Plugins/Editor/MeshEditor/Source/PolygonModeling/RemoveEdge.cpp
--- a/Engine/Plugins/Editor/MeshEditor/Source/PolygonModeling/RemoveEdge.cpp
+++ b/Engine/Plugins/Editor/MeshEditor/Source/PolygonModeling/RemoveEdge.cpp
@@ -6,7 +6,7 @@
 #include "Framework/Commands/UICommandInfo.h"
 #include "EditableMesh.h"
 #include "MeshElement.h"
-#include "Framework/Multibox/MultiBoxBuilder.h"
+#include "Framework/MultiBox/MultiBoxBuilder.h"
 #include "ScopedTransaction.h"
 
 #define LOCTEXT_NAMESPACE "MeshEditorMode"
diff --git a/Engine/Plugins/Runtime/EditableMesh/Source/EditableMesh/EditableMesh.cpp b/Engine/Plugins/Runtime/EditableMesh/Source/EditableMesh/EditableMesh.cpp
--- a/Engine/Plugins/Runtime/EditableMesh/Source/EditableMesh/EditableMesh.cpp
+++ b/Engine/Plugins/Runtime/EditableMesh/Source/EditableMesh/EditableMesh.cpp
@@ -21,14 +21,19 @@
 	#pragma warning(disable:4191)		// Disable warning C4191: 'type cast' : unsafe conversion
 #endif
 
+#ifndef M_PI
+#define M_PI_NOT_DEFINED
 #define M_PI PI		// OpenSubdiv is expecting M_PI to be defined already
+#endif
 
 #include "far/topologyRefinerFactory.h"
 #include "far/topologyDescriptor.h"
 #include "far/topologyRefiner.h"
 #include "far/primvarRefiner.h"
-	
+
+#ifdef M_PI_NOT_DEFINED
 #undef M_PI
+#endif
 
 #if defined(_MSC_VER)
 	#pragma warning(pop)
@@ -288,8 +293,8 @@ static void ApplyAttribute( TAttributesSet<ElementIDType>& AttributesSet, const
 	AttributeData.AttributeValue.Visit(
 		[ &AttributesSet, &AttributeData, ElementID ]( const auto& Value )
 		{
-			using AttributeType = TDecay<decltype( Value )>::Type;
-			AttributesSet.SetAttribute<AttributeType>( ElementID, AttributeData.AttributeName, AttributeData.AttributeIndex, Value );
+			using AttributeType = typename TDecay<decltype( Value )>::Type;
+			AttributesSet.template SetAttribute<AttributeType>( ElementID, AttributeData.AttributeName, AttributeData.AttributeIndex, Value );
 		}
 	);
 }
@@ -321,8 +326,8 @@ static void BackupAttributesInList( FMeshElementAttributeList& BackupList, const
 				const int32 AttributeIndex = AttributeData.AttributeIndex;
 
 				// Put current value into BackupList
-				using AttributeType = TDecay<decltype( Value )>::Type;
-				const AttributeType OldValue = AttributesSet.GetAttribute<AttributeType>( ElementID, AttributeName, AttributeIndex );
+				using AttributeType = typename TDecay<decltype( Value )>::Type;
+				const AttributeType OldValue = AttributesSet.template GetAttribute<AttributeType>( ElementID, AttributeName, AttributeIndex );
 				BackupList.Attributes.Emplace( AttributeName, AttributeIndex, FMeshElementAttributeValue( OldValue ) );
 			}
 		);
@@ -5862,7 +5867,7 @@ void UEditableMesh::ExtendVertices( const TArray<FVertexID>& VertexIDs, const bo
 						const int32 PerimeterVertexCount = GetPolygonPerimeterVertexCount( ConnectedPolygonID );
 						if( !( OriginalVertexNumberOnConnectedPolygon == ( PerimeterVertexCount - 1 ) && AdjacentVertexNumberOnConnectedPolygon == 0 ) &&  // Doesn't wrap forward?
 							( OriginalVertexNumberOnConnectedPolygon > AdjacentVertexNumberOnConnectedPolygon ||			// Winds backwards?
-							  AdjacentVertexNumberOnConnectedPolygon == ( PerimeterVertexCount - 1 ) && OriginalVertexNumberOnConnectedPolygon == 0 ) )	// Wraps backwards?
+							  ( AdjacentVertexNumberOnConnectedPolygon == ( PerimeterVertexCount - 1 ) && OriginalVertexNumberOnConnectedPolygon == 0 ) ) )	// Wraps backwards?
 						{
 							bConnectedPolygonWindsForward = false;
 						}
diff --git a/Engine/Source/Runtime/MeshBuilder/Private/BuildStatisticManager.h b/Engine/Source/Runtime/MeshBuilder/Private/BuildStatisticManager.h
--- a/Engine/Source/Runtime/MeshBuilder/Private/BuildStatisticManager.h
+++ b/Engine/Source/Runtime/MeshBuilder/Private/BuildStatisticManager.h
@@ -23,8 +23,8 @@ namespace BuildStatisticManager
 	{
 	public:
 		FBuildStatisticScope(const FString& InTimerDescription, FStatisticData& InStatisticData)
-			: StatisticData(InStatisticData)
-			, TimerDescription(InTimerDescription)
+			: TimerDescription(InTimerDescription)
+			, StatisticData(InStatisticData)
 		{
 			StartTime = FPlatformTime::Cycles();
 		}
@@ -36,8 +36,8 @@ namespace BuildStatisticManager
 			UE_LOG(LogMeshDescriptionBuildStatistic, Log, TEXT("%s: %f seconds - Counter: %d - Total Time: %f"), *TimerDescription, ScopeTime, StatisticData.Counter, StatisticData.TotalTime);
 		}
 	private:
-		int32 StartTime;
 		FString TimerDescription;
+		int32 StartTime;
 		FStatisticData& StatisticData;
 	};
 }
\ No newline at end of file
