Made subdivided mesh (using OpenSubdiv) a platform-conditional feature. We can enable it for more platforms once the libs are available. This will also fix the Mac build for now.

diff --git a/Engine/Plugins/Editor/MeshEditor/Source/MeshEditor/MeshEditorMode.cpp b/Engine/Plugins/Editor/MeshEditor/Source/MeshEditor/MeshEditorMode.cpp
--- a/Engine/Plugins/Editor/MeshEditor/Source/MeshEditor/MeshEditorMode.cpp
+++ b/Engine/Plugins/Editor/MeshEditor/Source/MeshEditor/MeshEditorMode.cpp
@@ -654,6 +654,7 @@ void FMeshEditorMode::BindCommands()
 	RegisterAnyElementCommand( MeshEditorAnyElementCommands.DeleteMeshElement, FExecuteAction::CreateLambda( [this] { DeleteSelectedMeshElement(); } ) );
 
 	// Register commands which work even without a selected element, as long as at least one mesh is selected
+#if EDITABLE_MESH_USE_OPENSUBDIV
 	RegisterCommonCommand( MeshEditorCommonCommands.AddSubdivisionLevel, FExecuteAction::CreateLambda( [this] { AddOrRemoveSubdivisionLevel( true ); } ) );
 	RegisterCommonCommand( MeshEditorCommonCommands.RemoveSubdivisionLevel, 
 		FExecuteAction::CreateLambda( [this] { AddOrRemoveSubdivisionLevel( false ); } ),
@@ -673,6 +674,7 @@ void FMeshEditorMode::BindCommands()
 			}
 		) 
 	);
+#endif
 
 	// @todo mesheditor: support EUserInterfaceActionType::ToggleButton actions in the UI, and extend RegisterCommand to allow
 	// a delegate returning check state.
@@ -1862,6 +1864,7 @@ bool FMeshEditorMode::DeleteSelectedMeshElement()
 }
 
 
+#if EDITABLE_MESH_USE_OPENSUBDIV
 void FMeshEditorMode::AddOrRemoveSubdivisionLevel( const bool bShouldAdd )
 {
 	if( ActiveAction == NAME_None )
@@ -1899,6 +1902,7 @@ void FMeshEditorMode::AddOrRemoveSubdivisionLevel( const bool bShouldAdd )
 		}
 	}
 }
+#endif
 
 
 void FMeshEditorMode::FrameSelectedElements( FEditorViewportClient* ViewportClient )
@@ -4909,6 +4913,7 @@ void FMeshEditorMode::OnActorSelectionChanged( const TArray<UObject*>& NewSelect
 
 void FMeshEditorMode::MakeVRRadialMenuActionsMenu(FMenuBuilder& MenuBuilder, TSharedPtr<FUICommandList> CommandList, class UVREditorMode* VRMode, float& RadiusOverride)
 {
+#if EDITABLE_MESH_USE_OPENSUBDIV
 	MenuBuilder.AddMenuEntry(
 		LOCTEXT("AddSubdivision", "Add SubD"),
 		FText(),
@@ -4931,6 +4936,7 @@ void FMeshEditorMode::MakeVRRadialMenuActionsMenu(FMenuBuilder& MenuBuilder, TSh
 		NAME_None,
 		EUserInterfaceActionType::ToggleButton
 		);
+#endif
 	MenuBuilder.AddMenuEntry(
 		LOCTEXT("EditInstance", "Edit Instance"),
 		FText(),
diff --git a/Engine/Plugins/Editor/MeshEditor/Source/MeshEditor/MeshEditorMode.h b/Engine/Plugins/Editor/MeshEditor/Source/MeshEditor/MeshEditorMode.h
--- a/Engine/Plugins/Editor/MeshEditor/Source/MeshEditor/MeshEditorMode.h
+++ b/Engine/Plugins/Editor/MeshEditor/Source/MeshEditor/MeshEditorMode.h
@@ -285,8 +285,10 @@ protected:
 	/** Deletes selected polygons, or polygons partly defined by selected elements; returns whether successful */
 	bool DeleteSelectedMeshElement();
 
+#if EDITABLE_MESH_USE_OPENSUBDIV
 	/** Adds or removes a subdivision level for selected meshes */
 	void AddOrRemoveSubdivisionLevel( const bool bShouldAdd );
+#endif
 
 	/** Moves the viewport camera to frame the currently selected elements */
 	void FrameSelectedElements( FEditorViewportClient* ViewportClient );
diff --git a/Engine/Plugins/Runtime/EditableMesh/Source/EditableMesh/EditableMesh.cpp b/Engine/Plugins/Runtime/EditableMesh/Source/EditableMesh/EditableMesh.cpp
--- a/Engine/Plugins/Runtime/EditableMesh/Source/EditableMesh/EditableMesh.cpp
+++ b/Engine/Plugins/Runtime/EditableMesh/Source/EditableMesh/EditableMesh.cpp
@@ -16,6 +16,8 @@
 // OpenSubdiv support
 // =========================================================
 
+#if EDITABLE_MESH_USE_OPENSUBDIV
+
 #if defined(_MSC_VER)
 	#pragma warning(push)
 	#pragma warning(disable:4191)		// Disable warning C4191: 'type cast' : unsafe conversion
@@ -39,6 +41,10 @@
 	#pragma warning(pop)
 #endif
 
+#endif
+
+// =========================================================
+
 DEFINE_LOG_CATEGORY( LogEditableMesh );
 
 // =========================================================
@@ -361,7 +367,9 @@ void UEditableMesh::PostLoad()
 {
 	Super::PostLoad();
 
+#if EDITABLE_MESH_USE_OPENSUBDIV
 	RefreshOpenSubdiv();
+#endif
 	RebuildOctree();
 	RebuildRenderMesh();
 }
@@ -668,6 +676,7 @@ void UEditableMesh::EndModification( const bool bFromUndo )
 
 		if( CurrentModificationType == EMeshModificationType::Final || !bFromUndo )
 		{
+#if EDITABLE_MESH_USE_OPENSUBDIV
 			// Update subdivision limit surface
 			if( CurrentToplogyChange == EMeshTopologyChange::TopologyChange )
 			{
@@ -679,6 +688,7 @@ void UEditableMesh::EndModification( const bool bFromUndo )
 				// No topology change, so we can ask OpenSubdiv to quickly generate new limit surface geometry
 				GenerateOpenSubdivLimitSurfaceData();
 			}
+#endif
 		}
 
 		// Every so often, compact the data.
@@ -1813,6 +1823,7 @@ FVector UEditableMesh::ComputePolygonNormal( const FPolygonID PolygonID ) const
 }
 
 
+#if EDITABLE_MESH_USE_OPENSUBDIV
 void UEditableMesh::RefreshOpenSubdiv()
 {
 	OsdTopologyRefiner.Reset();
@@ -1972,6 +1983,7 @@ void UEditableMesh::RefreshOpenSubdiv()
 
 	GenerateOpenSubdivLimitSurfaceData();
 }
+#endif
 
 
 const FSubdivisionLimitData& UEditableMesh::GetSubdivisionLimitData() const
@@ -1980,6 +1992,7 @@ const FSubdivisionLimitData& UEditableMesh::GetSubdivisionLimitData() const
 }
 
 
+#if EDITABLE_MESH_USE_OPENSUBDIV
 void UEditableMesh::GenerateOpenSubdivLimitSurfaceData()
 {
 	SubdivisionLimitData = FSubdivisionLimitData();
@@ -2452,6 +2465,7 @@ void UEditableMesh::GenerateOpenSubdivLimitSurfaceData()
 		}
 	}
 }
+#endif
 
 
 void UEditableMesh::RetriangulatePolygons()
@@ -2655,6 +2669,7 @@ bool UEditableMesh::ComputeBarycentricWeightForPointOnPolygon( const FPolygonID
 
 void UEditableMesh::SetSubdivisionCount( const int32 NewSubdivisionCount )
 {
+#if EDITABLE_MESH_USE_OPENSUBDIV
 	// @todo mesheditor subdiv: Really, instead of a custom FChange type for this type of change, we could create a change that
 	// represents a property (or set of properties) being changed, and use Unreal reflection to update it.  This would be
 	// reusable for future properties.
@@ -2678,6 +2693,7 @@ void UEditableMesh::SetSubdivisionCount( const int32 NewSubdivisionCount )
 	}
 
 	AddUndo( MakeUnique<FSetSubdivisionCountChange>( MoveTemp( RevertInput ) ) );
+#endif
 }
 
 
diff --git a/Engine/Plugins/Runtime/EditableMesh/Source/EditableMesh/EditableStaticMeshAdapter.cpp b/Engine/Plugins/Runtime/EditableMesh/Source/EditableMesh/EditableStaticMeshAdapter.cpp
--- a/Engine/Plugins/Runtime/EditableMesh/Source/EditableMesh/EditableStaticMeshAdapter.cpp
+++ b/Engine/Plugins/Runtime/EditableMesh/Source/EditableMesh/EditableStaticMeshAdapter.cpp
@@ -440,7 +440,9 @@ void UEditableStaticMeshAdapter::InitEditableStaticMesh( UEditableMesh* Editable
 	EditableMesh->GenerateTangentsAndNormals();
 #endif
 
+#if EDITABLE_MESH_USE_OPENSUBDIV
 	EditableMesh->RefreshOpenSubdiv();
+#endif
 	EditableMesh->RebuildOctree();
 }
 
diff --git a/Engine/Plugins/Runtime/EditableMesh/Source/EditableMesh/Public/EditableMesh.h b/Engine/Plugins/Runtime/EditableMesh/Source/EditableMesh/Public/EditableMesh.h
--- a/Engine/Plugins/Runtime/EditableMesh/Source/EditableMesh/Public/EditableMesh.h
+++ b/Engine/Plugins/Runtime/EditableMesh/Source/EditableMesh/Public/EditableMesh.h
@@ -22,8 +22,16 @@ EDITABLEMESH_API DECLARE_LOG_CATEGORY_EXTERN( LogEditableMesh, Log, All );
 #endif
 
 
-// @todo mesheditor: Comment these classes and enums!
+// Specify which platforms currently support OpenSubdiv
+#if PLATFORM_WINDOWS
+#define EDITABLE_MESH_USE_OPENSUBDIV 1
+#else
+#define EDITABLE_MESH_USE_OPENSUBDIV 0
+#endif
 
+
+// @todo mesheditor: Comment these classes and enums!
+#if EDITABLE_MESH_USE_OPENSUBDIV
 namespace OpenSubdiv
 {
 	namespace v3_2_0
@@ -34,6 +42,7 @@ namespace OpenSubdiv
 		}
 	}
 }
+#endif
 
 
 UENUM( BlueprintType )
@@ -473,13 +482,17 @@ protected:
 public:
 	// @todo mesheditor: temporarily changed access to public so the adapter can call it when building the editable mesh from the source static mesh. Think about this some more.
 	/** Refreshes the entire OpenSubdiv state for this mesh and generates subdivision geometry (if the mesh is configured to have subdivision levels) */
+#if EDITABLE_MESH_USE_OPENSUBDIV
 	void RefreshOpenSubdiv();
+#endif
 
 protected:
 
+#if EDITABLE_MESH_USE_OPENSUBDIV
 	/** Generates new limit surface data and caches it (if the mesh is configured to have subdivision levels).  This can be used to update the
 	    limit surface after a non-topology-effective edit has happened */
 	void GenerateOpenSubdivLimitSurfaceData();
+#endif
 
 
 private:
@@ -537,10 +550,6 @@ public:
 	UPROPERTY( BlueprintReadOnly, Category="Editable Mesh" )
 	int32 TextureCoordinateCount;
 
-	/** How many levels to subdivide this mesh.  Zero will turn off subdivisions */
-	UPROPERTY( BlueprintReadOnly, Category="Editable Mesh" ) 
-	int32 SubdivisionCount;
-
 	/** List of polygons which need their tangent basis recalculating (and consequently their associated vertex instances) */
 	TSet<FPolygonID> PolygonsPendingNewTangentBasis;
 
@@ -566,6 +575,11 @@ public:
 	/** Data will be compacted after this many topology modifying actions. */
 	static const int32 CompactFrequency = 50;
 
+	/** How many levels to subdivide this mesh.  Zero will turn off subdivisions */
+	UPROPERTY( BlueprintReadOnly, Category="Editable Mesh" ) 
+	int32 SubdivisionCount;
+
+#if EDITABLE_MESH_USE_OPENSUBDIV
 	/** OpenSubdiv topology refiner object.  This is generated for meshes that have subdivision levels, and reused to generate new limit surfaces 
 	    when geometry is moved.  When the mesh's topology changes, this object is regenerated from scratch. */
 	TSharedPtr<OpenSubdiv::v3_2_0::Far::TopologyRefiner> OsdTopologyRefiner;
@@ -585,6 +599,7 @@ public:
 		int const* ValueIndices;
 	};
 	TArray<FOsdFVarChannel> OsdFVarChannels;
+#endif
 
 	/** The resulting limit surface geometry after GenerateOpenSubdivLimitSurfaceData() is called */
 	FSubdivisionLimitData SubdivisionLimitData;
