Build fixes for Win32, Mac.

diff --git a/Engine/Plugins/Runtime/EditableMesh/Source/EditableMesh/EditableStaticMeshAdapter.cpp b/Engine/Plugins/Runtime/EditableMesh/Source/EditableMesh/EditableStaticMeshAdapter.cpp
--- a/Engine/Plugins/Runtime/EditableMesh/Source/EditableMesh/EditableStaticMeshAdapter.cpp
+++ b/Engine/Plugins/Runtime/EditableMesh/Source/EditableMesh/EditableStaticMeshAdapter.cpp
@@ -489,7 +489,15 @@ bool UEditableStaticMeshAdapter::FindOrAddMaterial(UMaterialInterface* Material,
 			Counter++;
 		}
 	}
-	MaterialIndex  = StaticMesh->StaticMaterials.Add(FStaticMaterial(Material, MaterialName, MaterialName));
+
+	MaterialIndex  = StaticMesh->StaticMaterials.Add(FStaticMaterial(
+		Material,
+		MaterialName,
+#if WITH_EDITOR
+		MaterialName
+#endif
+	));
+
 	ImportedMaterialName = MaterialName;
 	return true;
 }
diff --git a/Engine/Plugins/Runtime/EditableMesh/Source/EditableMesh/Public/EditableMeshTypes.h b/Engine/Plugins/Runtime/EditableMesh/Source/EditableMesh/Public/EditableMeshTypes.h
--- a/Engine/Plugins/Runtime/EditableMesh/Source/EditableMesh/Public/EditableMeshTypes.h
+++ b/Engine/Plugins/Runtime/EditableMesh/Source/EditableMesh/Public/EditableMeshTypes.h
@@ -348,7 +348,7 @@ struct FMeshElementAttributeValue
 
 	/** Define constructors */
 	FMeshElementAttributeValue() : Type( EMeshElementAttributeType::None ) {}
-	explicit FMeshElementAttributeValue( const FVector4 Value ) : Type( EMeshElementAttributeType::FVector4 ), Value_FVector4( Value ) {}
+	explicit FMeshElementAttributeValue( const FVector4& Value ) : Type( EMeshElementAttributeType::FVector4 ), Value_FVector4( Value ) {}
 	explicit FMeshElementAttributeValue( const FVector Value ) : Type( EMeshElementAttributeType::FVector ), Value_FVector( Value ) {}
 	explicit FMeshElementAttributeValue( const FVector2D Value ) : Type( EMeshElementAttributeType::FVector2D ), Value_FVector2D( Value ) {}
 	explicit FMeshElementAttributeValue( const float Value ) : Type( EMeshElementAttributeType::Float ), Value_Float( Value ) {}
