Miscellaneous build fixes for various platforms.

diff --git a/Engine/Plugins/Runtime/EditableMesh/Source/EditableMesh/EditableMesh.cpp b/Engine/Plugins/Runtime/EditableMesh/Source/EditableMesh/EditableMesh.cpp
--- a/Engine/Plugins/Runtime/EditableMesh/Source/EditableMesh/EditableMesh.cpp
+++ b/Engine/Plugins/Runtime/EditableMesh/Source/EditableMesh/EditableMesh.cpp
@@ -7161,17 +7161,17 @@ void UEditableMesh::WeldVertices( const TArray<FVertexID>& VertexIDsToWeld, FVer
 	// This function takes a list of perimeter vertices and a list of vertices to be welded as input.
 	// It returns a tuple stating whether the result is valid, and the [first, last) range of vertices to be welded.
 	// (It will be invalid if there is more than one contiguous run of vertices to weld.)
-	auto GetPerimeterVertexRangeToWeld = []( const TArray<FVertexID>& PolygonVertexIDs, const TArray<FVertexID>& VertexIDsToWeld )
+	auto GetPerimeterVertexRangeToWeld = []( const TArray<FVertexID>& PolygonVertexIDs, const TArray<FVertexID>& VerticesToWeld )
 	{
 		bool bValid = true;
 		int32 StartIndex = INDEX_NONE;
 		int32 EndIndex = INDEX_NONE;
 
 		const int32 NumPolygonVertices = PolygonVertexIDs.Num();
-		bool bPrevVertexNeedsWelding = VertexIDsToWeld.Contains( PolygonVertexIDs[ NumPolygonVertices - 1 ] );
+		bool bPrevVertexNeedsWelding = VerticesToWeld.Contains( PolygonVertexIDs[ NumPolygonVertices - 1 ] );
 		for( int32 Index = 0; Index < NumPolygonVertices; ++Index )
 		{
-			const bool bThisVertexNeedsWelding = VertexIDsToWeld.Contains( PolygonVertexIDs[ Index ] );
+			const bool bThisVertexNeedsWelding = VerticesToWeld.Contains( PolygonVertexIDs[ Index ] );
 			if( !bPrevVertexNeedsWelding && bThisVertexNeedsWelding )
 			{
 				// Transition from 'doesn't need welding' to 'needs welding'
diff --git a/Engine/Plugins/Runtime/EditableMesh/Source/EditableMesh/EditableStaticMeshAdapter.cpp b/Engine/Plugins/Runtime/EditableMesh/Source/EditableMesh/EditableStaticMeshAdapter.cpp
--- a/Engine/Plugins/Runtime/EditableMesh/Source/EditableMesh/EditableStaticMeshAdapter.cpp
+++ b/Engine/Plugins/Runtime/EditableMesh/Source/EditableMesh/EditableStaticMeshAdapter.cpp
@@ -1116,6 +1116,12 @@ void UEditableStaticMeshAdapter::UpdateCollision()
 
 	UBodySetup* BodySetup = StaticMesh->BodySetup;
 
+	// Return of body setup creation failed
+	if( !BodySetup )
+	{
+		return;
+	}
+
 	// NOTE: We don't bother calling Modify() on the BodySetup as EndModification() will rebuild this guy after every undo
 	// BodySetup->Modify();
 
diff --git a/Engine/Source/Runtime/Engine/Private/StaticMesh.cpp b/Engine/Source/Runtime/Engine/Private/StaticMesh.cpp
--- a/Engine/Source/Runtime/Engine/Private/StaticMesh.cpp
+++ b/Engine/Source/Runtime/Engine/Private/StaticMesh.cpp
@@ -2559,7 +2559,6 @@ public:
 		MeshDescriptionContent.AddZeroed(LodNumber);
 		for (int32 LodIndex = 0; LodIndex < LodNumber; ++LodIndex)
 		{
-			StaticMesh->SourceModels[LodIndex];
 			FStaticMeshSourceModel& SourceModel = StaticMesh->SourceModels[LodIndex];
 			MeshDescriptionTable[LodIndex] = SourceModel.OriginalMeshDescription != nullptr;
 			if (MeshDescriptionTable[LodIndex])
