Fix CIS build

diff --git a/Engine/Plugins/Runtime/EditableMesh/Source/EditableMesh/EditableStaticMeshAdapter.cpp b/Engine/Plugins/Runtime/EditableMesh/Source/EditableMesh/EditableStaticMeshAdapter.cpp
--- a/Engine/Plugins/Runtime/EditableMesh/Source/EditableMesh/EditableStaticMeshAdapter.cpp
+++ b/Engine/Plugins/Runtime/EditableMesh/Source/EditableMesh/EditableStaticMeshAdapter.cpp
@@ -490,11 +490,11 @@ bool UEditableStaticMeshAdapter::FindOrAddMaterial(UMaterialInterface* Material,
 		}
 	}
 
-	MaterialIndex  = StaticMesh->StaticMaterials.Add(FStaticMaterial(
-		Material,
-		MaterialName,
-#if WITH_EDITOR
-		MaterialName
+	MaterialIndex = StaticMesh->StaticMaterials.Add(FStaticMaterial(
+		Material
+		, MaterialName
+#if WITH_EDITORONLY_DATA
+		, MaterialName
 #endif
 	));
 
diff --git a/Engine/Source/Runtime/MeshBuilder/Private/MeshBuilder.cpp b/Engine/Source/Runtime/MeshBuilder/Private/MeshBuilder.cpp
--- a/Engine/Source/Runtime/MeshBuilder/Private/MeshBuilder.cpp
+++ b/Engine/Source/Runtime/MeshBuilder/Private/MeshBuilder.cpp
@@ -8,7 +8,6 @@
 #include "MeshDescription.h"
 #include "MeshAttributes.h"
 #include "RawMesh.h"
-#include "LayoutUV.h"
 
 
 FMeshBuilder::FMeshBuilder()
