Only define USE_OPENSUBDIV for Win64. Fix for Win32 build.

diff --git a/Engine/Plugins/Runtime/EditableMesh/Source/EditableMesh/Public/EditableMesh.h b/Engine/Plugins/Runtime/EditableMesh/Source/EditableMesh/Public/EditableMesh.h
--- a/Engine/Plugins/Runtime/EditableMesh/Source/EditableMesh/Public/EditableMesh.h
+++ b/Engine/Plugins/Runtime/EditableMesh/Source/EditableMesh/Public/EditableMesh.h
@@ -23,7 +23,7 @@ EDITABLEMESH_API DECLARE_LOG_CATEGORY_EXTERN( LogEditableMesh, Log, All );
 
 
 // Specify which platforms currently support OpenSubdiv
-#if PLATFORM_WINDOWS
+#if PLATFORM_WINDOWS && PLATFORM_64BITS
 #define EDITABLE_MESH_USE_OPENSUBDIV 1
 #else
 #define EDITABLE_MESH_USE_OPENSUBDIV 0
