Fixed issue with merging identical vertex instances. Now adds candidates for merging when vertex instance attributes are modified, and removes them when the parent vertex is deleted.

diff --git a/Engine/Plugins/Runtime/EditableMesh/Source/EditableMesh/EditableMesh.cpp b/Engine/Plugins/Runtime/EditableMesh/Source/EditableMesh/EditableMesh.cpp
--- a/Engine/Plugins/Runtime/EditableMesh/Source/EditableMesh/EditableMesh.cpp
+++ b/Engine/Plugins/Runtime/EditableMesh/Source/EditableMesh/EditableMesh.cpp
@@ -3749,6 +3749,7 @@ void UEditableMesh::DeleteOrphanVertices( const TArray<FVertexID>& VertexIDsToDe
 	{
 		for( const FVertexID VertexIDToDelete : VertexIDsToDelete )
 		{
+			VerticesPendingMerging.Remove( VertexIDToDelete );
 			GetMeshDescription()->DeleteVertex( VertexIDToDelete );
 		}
 	}
@@ -3761,8 +3762,7 @@ void UEditableMesh::DeleteVertexInstances( const TArray<FVertexInstanceID>& Vert
 {
 	UE_LOG( LogEditableMesh, Verbose, TEXT( "DeleteVertexInstances: %s" ), *LogHelpers::ArrayToString( VertexInstanceIDsToDelete ) );
 
-	FVertexInstanceArray& VertexInstances = GetMeshDescription()->VertexInstances();
-	FVertexArray& Vertices = GetMeshDescription()->Vertices();
+	UMeshDescription* MeshDescription = GetMeshDescription();
 
 	// Back everything up
 	{
@@ -3773,13 +3773,12 @@ void UEditableMesh::DeleteVertexInstances( const TArray<FVertexInstanceID>& Vert
 		for( int32 VertexNumber = VertexInstanceIDsToDelete.Num() - 1; VertexNumber >= 0; --VertexNumber )
 		{
 			const FVertexInstanceID VertexInstanceID = VertexInstanceIDsToDelete[ VertexNumber ];
-			const FMeshVertexInstance& VertexInstance = VertexInstances[ VertexInstanceID ];
 
 			// Back up properties
 			RevertInput.VertexInstancesToCreate.Emplace();
 			FVertexInstanceToCreate& VertexInstanceToCreate = RevertInput.VertexInstancesToCreate.Last();
 
-			VertexInstanceToCreate.VertexID = VertexInstance.VertexID;
+			VertexInstanceToCreate.VertexID = MeshDescription->GetVertexInstanceVertex( VertexInstanceID );
 			VertexInstanceToCreate.OriginalVertexInstanceID = VertexInstanceID;
 
 			BackupAllAttributes( VertexInstanceToCreate.VertexInstanceAttributes, GetMeshDescription()->VertexInstanceAttributes(), VertexInstanceID );
@@ -3801,6 +3800,7 @@ void UEditableMesh::DeleteVertexInstances( const TArray<FVertexInstanceID>& Vert
 	{
 		for( const FVertexInstanceID VertexInstanceIDToDelete : VertexInstanceIDsToDelete )
 		{
+			VerticesPendingMerging.Add( MeshDescription->GetVertexInstanceVertex( VertexInstanceIDToDelete ) );
 			GetMeshDescription()->DeleteVertexInstance( VertexInstanceIDToDelete, bDeleteOrphanedVertices ? &OrphanedVertexIDs : nullptr );
 		}
 	}
@@ -3987,6 +3987,7 @@ void UEditableMesh::CreateVertexInstances( const TArray<FVertexInstanceToCreate>
 			}
 
 			OutNewVertexInstanceIDs.Add( VertexInstanceID );
+			VerticesPendingMerging.Add( VertexInstanceToCreate.VertexID );
 		}
 	}
 
@@ -4653,6 +4654,8 @@ void UEditableMesh::SetVertexInstancesAttributes( const TArray<FAttributesForVer
 			// Set the new attribute
 			SetVertexInstanceAttribute( VertexInstanceID, VertexInstanceAttribute );
 		}
+
+		VerticesPendingMerging.Add( GetMeshDescription()->GetVertexInstanceVertex( VertexInstanceID ) );
 	}
 
 	AddUndo( MakeUnique<FSetVertexInstancesAttributesChange>( MoveTemp( RevertInput ) ) );
@@ -4951,7 +4954,9 @@ void UEditableMesh::SetPolygonContourVertexAttributes( FMeshPolygonContour& Cont
 
 			ChangePolygonsVertexInstances( VertexInstancesToChange );
 
-			// @todo mesheditor: weld identical vertex instances in the same smoothing group
+			// Weld identical vertex instances in the same smoothing group.
+			// We may need to do this after splitting a vertex instance if the result of the split created a vertex instance equal to another one.
+			VerticesPendingMerging.Add( VertexID );
 		}
 	}
 }
