Cleanup the polygon group attributes to use only the imported material slot name

diff --git a/Engine/Plugins/Editor/MeshEditor/Source/MeshEditor/MeshEditorMode.cpp b/Engine/Plugins/Editor/MeshEditor/Source/MeshEditor/MeshEditorMode.cpp
--- a/Engine/Plugins/Editor/MeshEditor/Source/MeshEditor/MeshEditorMode.cpp
+++ b/Engine/Plugins/Editor/MeshEditor/Source/MeshEditor/MeshEditorMode.cpp
@@ -2340,13 +2340,12 @@ bool FMeshEditorMode::AssignMaterialToSelectedPolygons( UMaterialInterface* Sele
 				// further details if there is more than one polygon group which matches the material.
 				FPolygonGroupID PolygonGroupToAssign = FPolygonGroupID::Invalid;
 
-				const TPolygonGroupAttributeArray<int>& PolygonGroupMaterialIndexes = MeshDescription->PolygonGroupAttributes().GetAttributes<int>(MeshAttribute::PolygonGroup::MaterialIndex);
 				const TPolygonGroupAttributeArray<FName>& PolygonGroupImportedMaterialSlotName = MeshDescription->PolygonGroupAttributes().GetAttributes<FName>(MeshAttribute::PolygonGroup::ImportedMaterialSlotName);
 
 				for( const FPolygonGroupID PolygonGroupID : MeshDescription->PolygonGroups().GetElementIDs() )
 				{
-					int32 PolygonGroupMaterialIndex = PolygonGroupMaterialIndexes[PolygonGroupID];
-					if(MaterialIndex == PolygonGroupMaterialIndex)
+					FName PolygonGroupMaterialName = PolygonGroupImportedMaterialSlotName[PolygonGroupID];
+					if(ImportedMaterialName == PolygonGroupMaterialName)
 					{
 						// We only expect to find one polygon group containing this material at the moment.
 						// We need to provide a way of distinguishing different polygon groups with the same material.
@@ -2365,7 +2364,6 @@ bool FMeshEditorMode::AssignMaterialToSelectedPolygons( UMaterialInterface* Sele
 					
 					FPolygonGroupToCreate& PolygonGroupToCreate = PolygonGroupsToCreate.Last();
 					PolygonGroupToCreate.PolygonGroupAttributes.Attributes.Emplace(MeshAttribute::PolygonGroup::ImportedMaterialSlotName, 0, FMeshElementAttributeValue(ImportedMaterialName));
-					PolygonGroupToCreate.PolygonGroupAttributes.Attributes.Emplace(MeshAttribute::PolygonGroup::MaterialIndex, 0, FMeshElementAttributeValue(MaterialIndex));
 					static TArray<FPolygonGroupID> NewPolygonGroupIDs;
 					EditableMesh->CreatePolygonGroups( PolygonGroupsToCreate, NewPolygonGroupIDs);
 					PolygonGroupToAssign = NewPolygonGroupIDs[ 0 ];
diff --git a/Engine/Plugins/Runtime/EditableMesh/Source/EditableMesh/EditableStaticMeshAdapter.cpp b/Engine/Plugins/Runtime/EditableMesh/Source/EditableMesh/EditableStaticMeshAdapter.cpp
--- a/Engine/Plugins/Runtime/EditableMesh/Source/EditableMesh/EditableStaticMeshAdapter.cpp
+++ b/Engine/Plugins/Runtime/EditableMesh/Source/EditableMesh/EditableStaticMeshAdapter.cpp
@@ -104,8 +104,6 @@ void UEditableStaticMeshAdapter::InitEditableStaticMesh( UEditableMesh* Editable
 
 				TEdgeAttributeArray<bool>& EdgeHardnesses = MeshDescription->EdgeAttributes().GetAttributes<bool>( MeshAttribute::Edge::IsHard );
 
-				TPolygonGroupAttributeArray<int>& PolygonGroupMaterialIndexes = MeshDescription->PolygonGroupAttributes().GetAttributes<int>(MeshAttribute::PolygonGroup::MaterialIndex);
-				TPolygonGroupAttributeArray<FName>& PolygonGroupMaterialSlotNames = MeshDescription->PolygonGroupAttributes().GetAttributes<FName>( MeshAttribute::PolygonGroup::MaterialSlotName );
 				TPolygonGroupAttributeArray<FName>& PolygonGroupImportedMaterialSlotNames = MeshDescription->PolygonGroupAttributes().GetAttributes<FName>( MeshAttribute::PolygonGroup::ImportedMaterialSlotName );
 				TPolygonGroupAttributeArray<bool>& PolygonGroupCollision = MeshDescription->PolygonGroupAttributes().GetAttributes<bool>( MeshAttribute::PolygonGroup::EnableCollision );
 				TPolygonGroupAttributeArray<bool>& PolygonGroupCastShadow = MeshDescription->PolygonGroupAttributes().GetAttributes<bool>( MeshAttribute::PolygonGroup::CastShadow );
@@ -276,8 +274,6 @@ void UEditableStaticMeshAdapter::InitEditableStaticMesh( UEditableMesh* Editable
 
 					// Create a new polygon group
 					const FPolygonGroupID NewPolygonGroupID = MeshDescription->CreatePolygonGroup();
-					PolygonGroupMaterialIndexes[NewPolygonGroupID] = RenderingSection.MaterialIndex;
-					PolygonGroupMaterialSlotNames[ NewPolygonGroupID ] = StaticMaterial.MaterialSlotName;
 					PolygonGroupImportedMaterialSlotNames[ NewPolygonGroupID ] = StaticMaterial.ImportedMaterialSlotName;
 					PolygonGroupCollision[ NewPolygonGroupID ] = RenderingSection.bEnableCollision;
 					PolygonGroupCastShadow[ NewPolygonGroupID ] = RenderingSection.bCastShadow;
@@ -1713,8 +1709,6 @@ void UEditableStaticMeshAdapter::OnCreatePolygonGroups( const UEditableMesh* Edi
 	UMeshDescription* MeshDescription = EditableMesh->GetMeshDescription();
 	check( MeshDescription );
 
-	TPolygonGroupAttributeArray<int>& PolygonGroupMaterialIndexes = MeshDescription->PolygonGroupAttributes().GetAttributes<int>(MeshAttribute::PolygonGroup::MaterialIndex);
-	const TPolygonGroupAttributeArray<FName>& PolygonGroupMaterialSlotNames = MeshDescription->PolygonGroupAttributes().GetAttributes<FName>( MeshAttribute::PolygonGroup::MaterialSlotName );
 	const TPolygonGroupAttributeArray<FName>& PolygonGroupImportedMaterialSlotNames = MeshDescription->PolygonGroupAttributes().GetAttributes<FName>( MeshAttribute::PolygonGroup::ImportedMaterialSlotName );
 	const TPolygonGroupAttributeArray<bool>& PolygonGroupCollision = MeshDescription->PolygonGroupAttributes().GetAttributes<bool>( MeshAttribute::PolygonGroup::EnableCollision );
 	const TPolygonGroupAttributeArray<bool>& PolygonGroupCastShadow = MeshDescription->PolygonGroupAttributes().GetAttributes<bool>( MeshAttribute::PolygonGroup::CastShadow );
@@ -1727,12 +1721,14 @@ void UEditableStaticMeshAdapter::OnCreatePolygonGroups( const UEditableMesh* Edi
 			//Verify the meshdescription new polygon group data is correct, The attribute must match the static mesh material slot array
 			//The staticmesh material slot entry was create before executing this function
 			bool bFoundMaterial = false;
+			int32 PolygonGroupMaterialIndex = INDEX_NONE;
 			for (int32 MaterialIndex = 0; MaterialIndex < StaticMesh->StaticMaterials.Num(); ++MaterialIndex)
 			{
 				if (PolygonGroupImportedMaterialSlotNames[PolygonGroupID] == StaticMesh->StaticMaterials[MaterialIndex].ImportedMaterialSlotName)
 				{
+					PolygonGroupMaterialIndex = MaterialIndex;
 					bFoundMaterial = true;
-					check(PolygonGroupMaterialIndexes[PolygonGroupID] == MaterialIndex);
+					break;
 				}
 			}
 			check(bFoundMaterial);
@@ -1768,7 +1764,7 @@ void UEditableStaticMeshAdapter::OnCreatePolygonGroups( const UEditableMesh* Edi
 			StaticMeshSection.MaxVertexIndex = 0;
 			StaticMeshSection.bEnableCollision = PolygonGroupCollision[ PolygonGroupID ];
 			StaticMeshSection.bCastShadow = PolygonGroupCastShadow[ PolygonGroupID ];
-			StaticMeshSection.MaterialIndex = PolygonGroupMaterialIndexes[ PolygonGroupID ];
+			StaticMeshSection.MaterialIndex = PolygonGroupMaterialIndex;
 		}
 
 		// Insert the rendering polygon group for keeping track of these index buffer properties
@@ -1790,21 +1786,25 @@ void UEditableStaticMeshAdapter::OnSetPolygonGroupAttribute( const UEditableMesh
 	FStaticMeshSection& StaticMeshSection = StaticMeshLOD.Sections[ RenderingPolygonGroup.RenderingSectionIndex ];
 
 	
-	const TPolygonGroupAttributeArray<int>& PolygonGroupMaterialIndexes = MeshDescription->PolygonGroupAttributes().GetAttributes<int>(MeshAttribute::PolygonGroup::MaterialIndex);
-	const TPolygonGroupAttributeArray<FName>& PolygonGroupMaterialSlotNames = MeshDescription->PolygonGroupAttributes().GetAttributes<FName>( MeshAttribute::PolygonGroup::MaterialSlotName );
 	const TPolygonGroupAttributeArray<FName>& PolygonGroupImportedMaterialSlotNames = MeshDescription->PolygonGroupAttributes().GetAttributes<FName>( MeshAttribute::PolygonGroup::ImportedMaterialSlotName );
 	
-	const int32 MaterialIndex = PolygonGroupMaterialIndexes[PolygonGroupID];
-
-	if (Attribute.AttributeName == MeshAttribute::PolygonGroup::MaterialIndex)
+	bool bFoundMaterial = false;
+	int32 PolygonGroupMaterialIndex = INDEX_NONE;
+	for (int32 MaterialIndex = 0; MaterialIndex < StaticMesh->StaticMaterials.Num(); ++MaterialIndex)
 	{
-		StaticMeshSection.MaterialIndex = MaterialIndex;
+		if (PolygonGroupImportedMaterialSlotNames[PolygonGroupID] == StaticMesh->StaticMaterials[MaterialIndex].ImportedMaterialSlotName)
+		{
+			PolygonGroupMaterialIndex = MaterialIndex;
+			bFoundMaterial = true;
+			break;
+		}
 	}
-	else if ( Attribute.AttributeName == MeshAttribute::PolygonGroup::MaterialSlotName ||
-		Attribute.AttributeName == MeshAttribute::PolygonGroup::ImportedMaterialSlotName )
+	check(bFoundMaterial);
+
+	const int32 MaterialIndex = PolygonGroupMaterialIndex;
+
+	if ( Attribute.AttributeName == MeshAttribute::PolygonGroup::ImportedMaterialSlotName )
 	{
-		
-		StaticMesh->StaticMaterials[MaterialIndex].MaterialSlotName = PolygonGroupMaterialSlotNames[PolygonGroupID];
 		StaticMesh->StaticMaterials[MaterialIndex].ImportedMaterialSlotName = PolygonGroupImportedMaterialSlotNames[PolygonGroupID];
 	}
 	else if( Attribute.AttributeName == MeshAttribute::PolygonGroup::CastShadow )
@@ -1820,17 +1820,17 @@ void UEditableStaticMeshAdapter::OnSetPolygonGroupAttribute( const UEditableMesh
 
 void UEditableStaticMeshAdapter::OnDeletePolygonGroups( const UEditableMesh* EditableMesh, const TArray<FPolygonGroupID>& PolygonGroupIDs )
 {
-	const TPolygonGroupAttributeArray<FName>& PolygonGroupMaterialSlotNames = EditableMesh->GetMeshDescription()->PolygonGroupAttributes().GetAttributes<FName>( MeshAttribute::PolygonGroup::MaterialSlotName );
+	const TPolygonGroupAttributeArray<FName>& PolygonGroupImportedMaterialSlotNames = EditableMesh->GetMeshDescription()->PolygonGroupAttributes().GetAttributes<FName>( MeshAttribute::PolygonGroup::ImportedMaterialSlotName );
 
 	for( const FPolygonGroupID PolygonGroupID : PolygonGroupIDs )
 	{
 		FRenderingPolygonGroup& RenderingPolygonGroup = RenderingPolygonGroups[ PolygonGroupID ];
 
-		const FName SlotName = PolygonGroupMaterialSlotNames[ PolygonGroupID ];
+		const FName SlotName = PolygonGroupImportedMaterialSlotNames[ PolygonGroupID ];
 
 		// Remove material slot associated with section
 		const int32 MaterialIndex = StaticMesh->StaticMaterials.IndexOfByPredicate(
-			[ &SlotName ]( const FStaticMaterial& StaticMaterial ) { return StaticMaterial.MaterialSlotName == SlotName; }
+			[ &SlotName ]( const FStaticMaterial& StaticMaterial ) { return StaticMaterial.ImportedMaterialSlotName == SlotName; }
 		);
 		StaticMesh->StaticMaterials.RemoveAt( MaterialIndex );
 
diff --git a/Engine/Source/Editor/UnrealEd/Private/Fbx/FbxStaticMeshImport.cpp b/Engine/Source/Editor/UnrealEd/Private/Fbx/FbxStaticMeshImport.cpp
--- a/Engine/Source/Editor/UnrealEd/Private/Fbx/FbxStaticMeshImport.cpp
+++ b/Engine/Source/Editor/UnrealEd/Private/Fbx/FbxStaticMeshImport.cpp
@@ -534,7 +534,6 @@ bool UnFbx::FFbxImporter::BuildStaticMeshFromGeometry(FbxNode* Node, UStaticMesh
 		TEdgeAttributeArray<float>& EdgeCreaseSharpnesses = MeshDescription->EdgeAttributes().GetAttributes<float>(MeshAttribute::Edge::CreaseSharpness);
 
 		TPolygonGroupAttributeArray<FName>& PolygonGroupImportedMaterialSlotNames = MeshDescription->PolygonGroupAttributes().GetAttributes<FName>(MeshAttribute::PolygonGroup::ImportedMaterialSlotName);
-		TPolygonGroupAttributeArray<int>& PolygonGroupMaterialIndex = MeshDescription->PolygonGroupAttributes().GetAttributes<int>(MeshAttribute::PolygonGroup::MaterialIndex);
 
 		int32 VertexOffset = MeshDescription->Vertices().Num();
 		int32 VertexInstanceOffset = MeshDescription->VertexInstances().Num();
@@ -815,7 +814,6 @@ bool UnFbx::FFbxImporter::BuildStaticMeshFromGeometry(FbxNode* Node, UStaticMesh
 				{
 					ExistingPolygonGroup = MeshDescription->CreatePolygonGroup();
 					PolygonGroupImportedMaterialSlotNames[ExistingPolygonGroup] = ImportedMaterialSlotName;
-					PolygonGroupMaterialIndex[ExistingPolygonGroup] = RealMaterialIndex;
 				}
 				PolygonGroupMapping.Add(RealMaterialIndex, ExistingPolygonGroup);
 			}
@@ -1894,12 +1892,11 @@ UStaticMesh* UnFbx::FFbxImporter::ImportStaticMeshAsSingle(UObject* InParent, TA
 		else
 		{
 			TPolygonGroupAttributeArray<FName>& PolygonGroupImportedMaterialSlotNames = MeshDescription->PolygonGroupAttributes().GetAttributes<FName>(MeshAttribute::PolygonGroup::ImportedMaterialSlotName);
-			TPolygonGroupAttributeArray<int>& PolygonGroupMaterialIndex = MeshDescription->PolygonGroupAttributes().GetAttributes<int>(MeshAttribute::PolygonGroup::MaterialIndex);
 
 			TArray<FStaticMaterial> MaterialToAdd;
 			for (const FPolygonGroupID PolygonGroupID : MeshDescription->PolygonGroups().GetElementIDs())
 			{
-				int32 MaterialIndex = PolygonGroupMaterialIndex[PolygonGroupID];
+				int32 MaterialIndex = PolygonGroupID.GetValue();
 				UMaterialInterface* Material = MeshMaterials.IsValidIndex(MaterialIndex) ? MeshMaterials[MaterialIndex].Material : UMaterial::GetDefaultMaterial(MD_Surface);
 				const FName& ImportedMaterialSlotName = PolygonGroupImportedMaterialSlotNames[PolygonGroupID];
 				const FName MaterialSlotName = ImportedMaterialSlotName;
diff --git a/Engine/Source/Editor/UnrealEd/Private/StaticMeshEdit.cpp b/Engine/Source/Editor/UnrealEd/Private/StaticMeshEdit.cpp
--- a/Engine/Source/Editor/UnrealEd/Private/StaticMeshEdit.cpp
+++ b/Engine/Source/Editor/UnrealEd/Private/StaticMeshEdit.cpp
@@ -717,7 +717,6 @@ void GetBrushMesh(ABrush* Brush, UModel* Model, UMeshDescription* MeshDescriptio
 	TEdgeAttributeArray<bool>& EdgeHardnesses = MeshDescription->EdgeAttributes().GetAttributes<bool>(MeshAttribute::Edge::IsHard);
 	TEdgeAttributeArray<float>& EdgeCreaseSharpnesses = MeshDescription->EdgeAttributes().GetAttributes<float>(MeshAttribute::Edge::CreaseSharpness);
 	TPolygonGroupAttributeArray<FName>& PolygonGroupImportedMaterialSlotNames = MeshDescription->PolygonGroupAttributes().GetAttributes<FName>(MeshAttribute::PolygonGroup::ImportedMaterialSlotName);
-	TPolygonGroupAttributeArray<int>& PolygonGroupMaterialIndex = MeshDescription->PolygonGroupAttributes().GetAttributes<int>(MeshAttribute::PolygonGroup::MaterialIndex);
 	
 	//Make sure we have one UVChannel
 	VertexInstanceUVs.SetNumIndices(1);
@@ -740,7 +739,7 @@ void GetBrushMesh(ABrush* Brush, UModel* Model, UMeshDescription* MeshDescriptio
 		FPolygonGroupID CurrentPolygonGroupID = FPolygonGroupID::Invalid;
 		for (const FPolygonGroupID& PolygonGroupID : MeshDescription->PolygonGroups().GetElementIDs())
 		{
-			if (MaterialIndex == PolygonGroupMaterialIndex[PolygonGroupID])
+			if (Material->GetFName() == PolygonGroupImportedMaterialSlotNames[PolygonGroupID])
 			{
 				CurrentPolygonGroupID = PolygonGroupID;
 				break;
@@ -749,7 +748,6 @@ void GetBrushMesh(ABrush* Brush, UModel* Model, UMeshDescription* MeshDescriptio
 		if (CurrentPolygonGroupID == FPolygonGroupID::Invalid)
 		{
 			CurrentPolygonGroupID = MeshDescription->CreatePolygonGroup();
-			PolygonGroupMaterialIndex[CurrentPolygonGroupID] = MaterialIndex;
 			PolygonGroupImportedMaterialSlotNames[CurrentPolygonGroupID] = Material->GetFName();
 		}
 
diff --git a/Engine/Source/Runtime/Engine/Private/StaticMesh.cpp b/Engine/Source/Runtime/Engine/Private/StaticMesh.cpp
--- a/Engine/Source/Runtime/Engine/Private/StaticMesh.cpp
+++ b/Engine/Source/Runtime/Engine/Private/StaticMesh.cpp
@@ -2502,7 +2502,7 @@ void UStaticMesh::FixupMaterialSlotName()
 // differences, etc.) replace the version GUID below with a new one.
 // In case of merge conflicts with DDC versions, you MUST generate a new GUID
 // and set this new GUID as the version.                                       
-#define MESHDATAKEY_STATICMESH_DERIVEDDATA_VER TEXT("F9C67C97D5854B2DB45CC1AF78135EE0")
+#define MESHDATAKEY_STATICMESH_DERIVEDDATA_VER TEXT("65133A83EADF412DB9C57EF087814E02")
 
 static const FString& GetMeshDataKeyStaticMeshDerivedDataVersion()
 {
diff --git a/Engine/Source/Runtime/MeshDescription/Private/MeshAttributes.cpp b/Engine/Source/Runtime/MeshDescription/Private/MeshAttributes.cpp
--- a/Engine/Source/Runtime/MeshDescription/Private/MeshAttributes.cpp
+++ b/Engine/Source/Runtime/MeshDescription/Private/MeshAttributes.cpp
@@ -21,10 +21,8 @@ namespace MeshAttribute
 	const FName Polygon::Binormal( "Binormal" );
 	const FName Polygon::Center( "Center" );
 
-	const FName PolygonGroup::MaterialSlotName( "MaterialSlotName" );
 	const FName PolygonGroup::ImportedMaterialSlotName( "ImportedMaterialSlotName" );
 	const FName PolygonGroup::EnableCollision( "EnableCollision" );
 	const FName PolygonGroup::CastShadow( "CastShadow" );
-	const FName PolygonGroup::MaterialIndex( "MaterialIndex" );
 }
 
diff --git a/Engine/Source/Runtime/MeshDescription/Private/MeshDescription.cpp b/Engine/Source/Runtime/MeshDescription/Private/MeshDescription.cpp
--- a/Engine/Source/Runtime/MeshDescription/Private/MeshDescription.cpp
+++ b/Engine/Source/Runtime/MeshDescription/Private/MeshDescription.cpp
@@ -36,8 +36,6 @@ void UMeshDescription::RegisterBaseAttributes()
 	PolygonAttributes().RegisterAttribute<FVector>(MeshAttribute::Polygon::Center, 1, FVector::ZeroVector, EMeshAttributeFlags::AutoGenerated);
 
 	// Add basic polygon group attributes
-	PolygonGroupAttributes().RegisterAttribute<int>(MeshAttribute::PolygonGroup::MaterialIndex); //This attribute represent the section material index
-	PolygonGroupAttributes().RegisterAttribute<FName>(MeshAttribute::PolygonGroup::MaterialSlotName); //Deprecated
 	PolygonGroupAttributes().RegisterAttribute<FName>(MeshAttribute::PolygonGroup::ImportedMaterialSlotName); //The unique key to match the mesh material slot
 	PolygonGroupAttributes().RegisterAttribute<bool>(MeshAttribute::PolygonGroup::EnableCollision); //Deprecated
 	PolygonGroupAttributes().RegisterAttribute<bool>(MeshAttribute::PolygonGroup::CastShadow); //Deprecated
@@ -67,8 +65,6 @@ void UMeshDescription::UnRegisterBaseAttributes()
 	PolygonAttributes().UnregisterAttribute<FVector>(MeshAttribute::Polygon::Center);
 
 	// Remove basic polygon group attributes
-	PolygonGroupAttributes().UnregisterAttribute<int>(MeshAttribute::PolygonGroup::MaterialIndex);
-	PolygonGroupAttributes().UnregisterAttribute<FName>(MeshAttribute::PolygonGroup::MaterialSlotName);
 	PolygonGroupAttributes().UnregisterAttribute<FName>(MeshAttribute::PolygonGroup::ImportedMaterialSlotName);
 	PolygonGroupAttributes().UnregisterAttribute<bool>(MeshAttribute::PolygonGroup::EnableCollision);
 	PolygonGroupAttributes().UnregisterAttribute<bool>(MeshAttribute::PolygonGroup::CastShadow);
@@ -130,8 +126,6 @@ FString UMeshDescription::GetIdString()
 	TPolygonAttributeArray<FVector>& PolygonBinormals = PolygonAttributes().GetAttributes<FVector>(MeshAttribute::Polygon::Binormal);
 	TPolygonAttributeArray<FVector>& PolygonCenters = PolygonAttributes().GetAttributes<FVector>(MeshAttribute::Polygon::Center);
 
-	TPolygonGroupAttributeArray<int>& PolygonGroupMaterialIndexes = PolygonGroupAttributes().GetAttributes<int>(MeshAttribute::PolygonGroup::MaterialIndex);
-	TPolygonGroupAttributeArray<FName>& PolygonGroupMaterialSlotNames = PolygonGroupAttributes().GetAttributes<FName>(MeshAttribute::PolygonGroup::MaterialSlotName);
 	TPolygonGroupAttributeArray<FName>& PolygonGroupImportedMaterialSlotNames = PolygonGroupAttributes().GetAttributes<FName>(MeshAttribute::PolygonGroup::ImportedMaterialSlotName);
 	TPolygonGroupAttributeArray<bool>& PolygonGroupEnableCollisions = PolygonGroupAttributes().GetAttributes<bool>(MeshAttribute::PolygonGroup::EnableCollision);
 	TPolygonGroupAttributeArray<bool>& PolygonGroupCastShadows = PolygonGroupAttributes().GetAttributes<bool>(MeshAttribute::PolygonGroup::CastShadow);
@@ -224,7 +218,8 @@ FString UMeshDescription::GetIdString()
 			{
 				Ar << VertexPositions[GetVertexInstance(VertexInstanceID).VertexID];
 			}
-			ItemCount = PolygonGroupMaterialIndexes[MeshItem.PolygonGroupID];
+			FPolygonGroupID PolygonGroupId = MeshItem.PolygonGroupID;
+			ItemCount = (uint32)PolygonGroupId.GetValue();
 			Ar.SerializeInt(ItemCount, 0);
 
 			Ar << GetPolygon(ArrayID);
@@ -244,8 +239,6 @@ FString UMeshDescription::GetIdString()
 			const FMeshPolygonGroup& MeshItem = GetPolygonGroup(ArrayID);
 			uint32 ItemCount = MeshItem.Polygons.Num();
 			Ar.SerializeInt(ItemCount, 0);
-			Ar << PolygonGroupMaterialIndexes[ArrayID];
-			Ar << PolygonGroupMaterialSlotNames[ArrayID];
 			Ar << PolygonGroupImportedMaterialSlotNames[ArrayID];
 			Ar << PolygonGroupEnableCollisions[ArrayID];
 			Ar << PolygonGroupCastShadows[ArrayID];
diff --git a/Engine/Source/Runtime/MeshDescription/Public/MeshAttributes.h b/Engine/Source/Runtime/MeshDescription/Public/MeshAttributes.h
--- a/Engine/Source/Runtime/MeshDescription/Public/MeshAttributes.h
+++ b/Engine/Source/Runtime/MeshDescription/Public/MeshAttributes.h
@@ -40,10 +40,8 @@ namespace MeshAttribute
 
 	namespace PolygonGroup
 	{
-		extern MESHDESCRIPTION_API const FName MaterialSlotName;
 		extern MESHDESCRIPTION_API const FName ImportedMaterialSlotName;
 		extern MESHDESCRIPTION_API const FName EnableCollision;
 		extern MESHDESCRIPTION_API const FName CastShadow;
-		extern MESHDESCRIPTION_API const FName MaterialIndex;
 	}
 }
