Fixed shadowed variable.

diff --git a/Engine/Plugins/Runtime/EditableMesh/Source/EditableMesh/EditableMesh.cpp b/Engine/Plugins/Runtime/EditableMesh/Source/EditableMesh/EditableMesh.cpp
--- a/Engine/Plugins/Runtime/EditableMesh/Source/EditableMesh/EditableMesh.cpp
+++ b/Engine/Plugins/Runtime/EditableMesh/Source/EditableMesh/EditableMesh.cpp
@@ -3762,8 +3762,6 @@ void UEditableMesh::DeleteVertexInstances( const TArray<FVertexInstanceID>& Vert
 {
 	UE_LOG( LogEditableMesh, Verbose, TEXT( "DeleteVertexInstances: %s" ), *LogHelpers::ArrayToString( VertexInstanceIDsToDelete ) );
 
-	UMeshDescription* MeshDescription = GetMeshDescription();
-
 	// Back everything up
 	{
 		FCreateVertexInstancesChangeInput RevertInput;
@@ -3778,7 +3776,7 @@ void UEditableMesh::DeleteVertexInstances( const TArray<FVertexInstanceID>& Vert
 			RevertInput.VertexInstancesToCreate.Emplace();
 			FVertexInstanceToCreate& VertexInstanceToCreate = RevertInput.VertexInstancesToCreate.Last();
 
-			VertexInstanceToCreate.VertexID = MeshDescription->GetVertexInstanceVertex( VertexInstanceID );
+			VertexInstanceToCreate.VertexID = GetMeshDescription()->GetVertexInstanceVertex( VertexInstanceID );
 			VertexInstanceToCreate.OriginalVertexInstanceID = VertexInstanceID;
 
 			BackupAllAttributes( VertexInstanceToCreate.VertexInstanceAttributes, GetMeshDescription()->VertexInstanceAttributes(), VertexInstanceID );
@@ -3800,7 +3798,7 @@ void UEditableMesh::DeleteVertexInstances( const TArray<FVertexInstanceID>& Vert
 	{
 		for( const FVertexInstanceID VertexInstanceIDToDelete : VertexInstanceIDsToDelete )
 		{
-			VerticesPendingMerging.Add( MeshDescription->GetVertexInstanceVertex( VertexInstanceIDToDelete ) );
+			VerticesPendingMerging.Add( GetMeshDescription()->GetVertexInstanceVertex( VertexInstanceIDToDelete ) );
 			GetMeshDescription()->DeleteVertexInstance( VertexInstanceIDToDelete, bDeleteOrphanedVertices ? &OrphanedVertexIDs : nullptr );
 		}
 	}
