Fixed Redo regression due to inconsistent identification of mergeable vertices.
Improved EDITABLE_MESH_DEBUG_LOG and created a pair of macros to perform hierarchical output.

diff --git a/Engine/Plugins/Editor/MeshEditor/Source/MeshEditor/MeshEditorMode.cpp b/Engine/Plugins/Editor/MeshEditor/Source/MeshEditor/MeshEditorMode.cpp
--- a/Engine/Plugins/Editor/MeshEditor/Source/MeshEditor/MeshEditorMode.cpp
+++ b/Engine/Plugins/Editor/MeshEditor/Source/MeshEditor/MeshEditorMode.cpp
@@ -1103,6 +1103,11 @@ void FMeshEditorMode::RollbackPreviewChanges()
 {
 	// NOTE: We iterate backwards here, because changes were added to our array in the order they originally
 	// happened.  But we'll need to apply their revert in the opposite order.
+	if( PreviewRevertChanges.Num() > 0 )
+	{
+		UE_LOG( LogEditableMesh, Verbose, TEXT( "------- ROLLING BACK PREVIEW CHANGE -------" ) );
+	}
+
 	for( int32 ChangeIndex = PreviewRevertChanges.Num() - 1; ChangeIndex >= 0; --ChangeIndex )
 	{
 		TTuple<UObject*, TUniquePtr<FChange>>& ObjectAndPreviewRevertChange = PreviewRevertChanges[ ChangeIndex ];
@@ -1120,15 +1125,21 @@ void FMeshEditorMode::RollbackPreviewChanges()
 //		PreviewRevertChange->PrintToLog( *GWarn );
 //		GWarn->Logf( TEXT( "---------- End (Object:%s) ----------" ), *Object->GetName() );
 
-		UE_LOG( LogEditableMesh, Verbose, TEXT( "------- ROLLING BACK PREVIEW CHANGE -------" ) );
+		UE_LOG( LogEditableMesh, Verbose, TEXT( "------- Transaction start -------" ) );
 		TUniquePtr<FChange> UnusedChangeToUndoRevert = PreviewRevertChange->Execute( Object );
-		UE_LOG( LogEditableMesh, Verbose, TEXT( "------- END ROLL BACK PREVIEW CHANGE -------" ) );
+		UE_LOG( LogEditableMesh, Verbose, TEXT( "------- Transaction end -------" ) );
 
 		// @todo mesheditor debug
 //		GWarn->Logf( TEXT( "-----(Here's what the Redo looks like)-----" ) );
 //		UnusedChangeToUndoRevert->PrintToLog( *GWarn );
 //		GWarn->Logf( TEXT( "-------------------------------------------" ) );
 	}
+
+	if( PreviewRevertChanges.Num() > 0 )
+	{
+		UE_LOG( LogEditableMesh, Verbose, TEXT( "------- END ROLL BACK PREVIEW CHANGE -------" ) );
+	}
+
 	PreviewRevertChanges.Reset();
 }
 
diff --git a/Engine/Plugins/Runtime/EditableMesh/Source/EditableMesh/EditableMesh.cpp b/Engine/Plugins/Runtime/EditableMesh/Source/EditableMesh/EditableMesh.cpp
--- a/Engine/Plugins/Runtime/EditableMesh/Source/EditableMesh/EditableMesh.cpp
+++ b/Engine/Plugins/Runtime/EditableMesh/Source/EditableMesh/EditableMesh.cpp
@@ -47,6 +47,18 @@
 
 DEFINE_LOG_CATEGORY( LogEditableMesh );
 
+#ifdef EDITABLE_MESH_DEBUG_LOG
+static int32 GIndent = 0;
+static FString IncreaseIndent() { GIndent += 2; return FString::ChrN( GIndent - 2, ' ' ); }
+static FString DecreaseIndent() { GIndent -= 2; return FString::ChrN( GIndent, ' ' ); }
+
+#define EM_ENTER(Text, ...) UE_LOG( LogEditableMesh, Verbose, TEXT( "%s" ) Text, *IncreaseIndent(), __VA_ARGS__ )
+#define EM_EXIT(Text, ...) UE_LOG( LogEditableMesh, Verbose, TEXT( "%s" ) Text, *DecreaseIndent(), __VA_ARGS__ )
+#else
+#define EM_ENTER(Text, ...)
+#define EM_EXIT(Text, ...)
+#endif
+
 // =========================================================
 
 
@@ -1185,7 +1197,7 @@ void UEditableMesh::GetVertexConnectedPolygonsInSameSoftEdgedGroup( const FVerte
 
 void UEditableMesh::SplitVertexInstanceInPolygons( const FVertexInstanceID VertexInstanceID, const TArray<FPolygonID>& PolygonIDs )
 {
-	UE_LOG( LogEditableMesh, Verbose, TEXT( "SplitVertexInstanceInPolygons: %s %s" ), *VertexInstanceID.ToString(), *LogHelpers::ArrayToString( PolygonIDs ) );
+	EM_ENTER( TEXT( "SplitVertexInstanceInPolygons: %s %s" ), *VertexInstanceID.ToString(), *LogHelpers::ArrayToString( PolygonIDs ) );
 
 	// Create a new vertex instance copied from the one passed in
 
@@ -1201,13 +1213,13 @@ void UEditableMesh::SplitVertexInstanceInPolygons( const FVertexInstanceID Verte
 
 	ReplaceVertexInstanceInPolygons( VertexInstanceID, NewVertexInstanceIDs[ 0 ], PolygonIDs );
 
-	UE_LOG( LogEditableMesh, Verbose, TEXT( "* SplitVertexInstanceInPolygons returned" ) );
+	EM_EXIT( TEXT( "SplitVertexInstanceInPolygons returned" ) );
 }
 
 
 void UEditableMesh::ReplaceVertexInstanceInPolygons( const FVertexInstanceID OldVertexInstanceID, const FVertexInstanceID NewVertexInstanceID, const TArray<FPolygonID>& PolygonIDs )
 {
-	UE_LOG( LogEditableMesh, Verbose, TEXT( "ReplaceVertexInstanceInPolygons: %s %s %s" ), *OldVertexInstanceID.ToString(), *NewVertexInstanceID.ToString(), *LogHelpers::ArrayToString( PolygonIDs ) );
+	EM_ENTER( TEXT( "ReplaceVertexInstanceInPolygons: %s %s %s" ), *OldVertexInstanceID.ToString(), *NewVertexInstanceID.ToString(), *LogHelpers::ArrayToString( PolygonIDs ) );
 
 	// Substitute the new vertex instance in the passed in polygons
 
@@ -1259,7 +1271,7 @@ void UEditableMesh::ReplaceVertexInstanceInPolygons( const FVertexInstanceID Old
 
 	ChangePolygonsVertexInstances( VertexInstancesToChange );
 
-	UE_LOG( LogEditableMesh, Verbose, TEXT( "* ReplaceVertexInstanceInPolygons returned" ) );
+	EM_EXIT( TEXT( "ReplaceVertexInstanceInPolygons returned" ) );
 }
 
 
@@ -1324,6 +1336,13 @@ void UEditableMesh::SplitVerticesIfNecessary( const TArray<FVertexID>& VerticesT
 
 void UEditableMesh::MergeVertexInstances()
 {
+	if( VerticesPendingMerging.Num() == 0 )
+	{
+		return;
+	}
+
+	EM_ENTER( TEXT( "MergeVertexInstances" ) );
+
 	static TArray<FVertexInstanceID> VertexInstancesToDelete;
 	VertexInstancesToDelete.Reset();
 
@@ -1399,6 +1418,8 @@ void UEditableMesh::MergeVertexInstances()
 	// Delete orphaned vertex instances
 	const bool bDeleteOrphanedVertices = false;
 	DeleteVertexInstances( VertexInstancesToDelete, bDeleteOrphanedVertices );
+
+	EM_EXIT( TEXT( "MergeVertexInstances returned" ) );
 }
 
 
@@ -2699,6 +2720,8 @@ void UEditableMesh::SetSubdivisionCount( const int32 NewSubdivisionCount )
 
 void UEditableMesh::MoveVertices( const TArray<FVertexToMove>& VerticesToMove )
 {
+	EM_ENTER( TEXT( "MoveVertices: %s" ), *LogHelpers::ArrayToString( VerticesToMove ) );
+
 	static TSet< FPolygonID > VertexConnectedPolygons;
 	VertexConnectedPolygons.Reset();
 
@@ -2731,6 +2754,8 @@ void UEditableMesh::MoveVertices( const TArray<FVertexToMove>& VerticesToMove )
 	// Everything needs to be retriangulated because convexity may have changed
 	PolygonsPendingNewTangentBasis.Append( VertexConnectedPolygons );
 	PolygonsPendingTriangulation.Append( VertexConnectedPolygons );
+
+	EM_EXIT( TEXT( "MoveVertices returned" ) );
 }
 
 
@@ -2868,7 +2893,7 @@ void UEditableMesh::CreateMissingPolygonHoleEdges( const FPolygonID PolygonID, c
 
 void UEditableMesh::SplitEdge( const FEdgeID EdgeID, const TArray<float>& Splits, TArray<FVertexID>& OutNewVertexIDs )
 {
-	UE_LOG( LogEditableMesh, Verbose, TEXT( "SplitEdge: %s" ), *EdgeID.ToString() );
+	EM_ENTER( TEXT( "SplitEdge: %s" ), *EdgeID.ToString() );
 
 	// NOTE: The incoming splits should always be between 0.0 and 1.0, representing progress along 
 	//       the edge from the edge's first vertex toward it's other vertex.  The order doesn't matter (we'll sort them.)
@@ -3088,7 +3113,7 @@ void UEditableMesh::SplitEdge( const FEdgeID EdgeID, const TArray<float>& Splits
 	// Retriangulate all of the affected polygons
 	PolygonsPendingTriangulation.Append( GetMeshDescription()->GetEdgeConnectedPolygons( EdgeID ) );
 
-	UE_LOG( LogEditableMesh, Verbose, TEXT( "* SplitEdge returned %s" ), *LogHelpers::ArrayToString( OutNewVertexIDs ) );
+	EM_EXIT( TEXT( "SplitEdge returned %s" ), *LogHelpers::ArrayToString( OutNewVertexIDs ) );
 }
 
 
@@ -3708,7 +3733,7 @@ void UEditableMesh::DeleteVertexAndConnectedEdgesAndPolygons( const FVertexID Ve
 
 void UEditableMesh::DeleteOrphanVertices( const TArray<FVertexID>& VertexIDsToDelete )
 {
-	UE_LOG( LogEditableMesh, Verbose, TEXT( "DeleteOrphanVertices: %s" ), *LogHelpers::ArrayToString( VertexIDsToDelete ) );
+	EM_ENTER( TEXT( "DeleteOrphanVertices: %s" ), *LogHelpers::ArrayToString( VertexIDsToDelete ) );
 
 	FVertexArray& Vertices = GetMeshDescription()->Vertices();
 
@@ -3754,13 +3779,13 @@ void UEditableMesh::DeleteOrphanVertices( const TArray<FVertexID>& VertexIDsToDe
 		}
 	}
 
-	UE_LOG( LogEditableMesh, Verbose, TEXT( "* DeleteOrphanVertices returned" ) );
+	EM_EXIT( TEXT( "DeleteOrphanVertices returned" ) );
 }
 
 
 void UEditableMesh::DeleteVertexInstances( const TArray<FVertexInstanceID>& VertexInstanceIDsToDelete, const bool bDeleteOrphanedVertices )
 {
-	UE_LOG( LogEditableMesh, Verbose, TEXT( "DeleteVertexInstances: %s" ), *LogHelpers::ArrayToString( VertexInstanceIDsToDelete ) );
+	EM_ENTER( TEXT( "DeleteVertexInstances: %s" ), *LogHelpers::ArrayToString( VertexInstanceIDsToDelete ) );
 
 	// Back everything up
 	{
@@ -3809,13 +3834,13 @@ void UEditableMesh::DeleteVertexInstances( const TArray<FVertexInstanceID>& Vert
 		DeleteOrphanVertices( OrphanedVertexIDs );
 	}
 
-	UE_LOG( LogEditableMesh, Verbose, TEXT( "* DeleteVertexInstances returned" ) );
+	EM_EXIT( TEXT( "DeleteVertexInstances returned" ) );
 }
 
 
 void UEditableMesh::DeleteEdges( const TArray<FEdgeID>& EdgeIDsToDelete, const bool bDeleteOrphanedVertices )
 {
-	UE_LOG( LogEditableMesh, Verbose, TEXT( "DeleteEdges: %s" ), *LogHelpers::ArrayToString( EdgeIDsToDelete ) );
+	EM_ENTER( TEXT( "DeleteEdges: %s" ), *LogHelpers::ArrayToString( EdgeIDsToDelete ) );
 
 	FEdgeArray& Edges = GetMeshDescription()->Edges();
 	FVertexArray& Vertices = GetMeshDescription()->Vertices();
@@ -3866,13 +3891,13 @@ void UEditableMesh::DeleteEdges( const TArray<FEdgeID>& EdgeIDsToDelete, const b
 		}
 	}
 
-	UE_LOG( LogEditableMesh, Verbose, TEXT( "* DeleteEdges returned" ) );
+	EM_EXIT( TEXT( "DeleteEdges returned" ) );
 }
 
 
 void UEditableMesh::CreateEmptyVertexRange( const int32 NumVerticesToCreate, TArray<FVertexID>& OutNewVertexIDs )
 {
-	UE_LOG( LogEditableMesh, Verbose, TEXT( "CreateEmptyVertexRange: %d" ), NumVerticesToCreate );
+	EM_ENTER( TEXT( "CreateEmptyVertexRange: %d" ), NumVerticesToCreate );
 
 	OutNewVertexIDs.Reset( NumVerticesToCreate );
 
@@ -3903,13 +3928,13 @@ void UEditableMesh::CreateEmptyVertexRange( const int32 NumVerticesToCreate, TAr
 		Adapter->OnCreateEmptyVertexRange( this, OutNewVertexIDs );
 	}
 
-	UE_LOG( LogEditableMesh, Verbose, TEXT( "* CreateEmptyVertexRange returned %s" ), *LogHelpers::ArrayToString( OutNewVertexIDs ) );
+	EM_EXIT( TEXT( "CreateEmptyVertexRange returned %s" ), *LogHelpers::ArrayToString( OutNewVertexIDs ) );
 }
 
 
 void UEditableMesh::CreateVertices( const TArray<FVertexToCreate>& VerticesToCreate, TArray<FVertexID>& OutNewVertexIDs )
 {
-	UE_LOG( LogEditableMesh, Verbose, TEXT( "CreateVertices: %s" ), *LogHelpers::ArrayToString( VerticesToCreate ) );
+	EM_ENTER( TEXT( "CreateVertices: %s" ), *LogHelpers::ArrayToString( VerticesToCreate ) );
 
 	// Create vertices
 	{
@@ -3959,13 +3984,13 @@ void UEditableMesh::CreateVertices( const TArray<FVertexToCreate>& VerticesToCre
 		}
 	}
 
-	UE_LOG( LogEditableMesh, Verbose, TEXT( "* CreateVertices returned %s" ), *LogHelpers::ArrayToString( OutNewVertexIDs ) );
+	EM_EXIT( TEXT( "CreateVertices returned %s" ), *LogHelpers::ArrayToString( OutNewVertexIDs ) );
 }
 
 
 void UEditableMesh::CreateVertexInstances( const TArray<FVertexInstanceToCreate>& VertexInstancesToCreate, TArray<FVertexInstanceID>& OutNewVertexInstanceIDs )
 {
-	UE_LOG( LogEditableMesh, Verbose, TEXT( "CreateVertexInstances: %s" ), *LogHelpers::ArrayToString( VertexInstancesToCreate ) );
+	EM_ENTER( TEXT( "CreateVertexInstances: %s" ), *LogHelpers::ArrayToString( VertexInstancesToCreate ) );
 
 	// Create new vertex instances
 	{
@@ -4017,13 +4042,13 @@ void UEditableMesh::CreateVertexInstances( const TArray<FVertexInstanceToCreate>
 		}
 	}
 
-	UE_LOG( LogEditableMesh, Verbose, TEXT( "* CreateVertexInstances returned %s" ), *LogHelpers::ArrayToString( OutNewVertexInstanceIDs ) );
+	EM_EXIT( TEXT( "CreateVertexInstances returned %s" ), *LogHelpers::ArrayToString( OutNewVertexInstanceIDs ) );
 }
 
 
 void UEditableMesh::CreateEdges( const TArray<FEdgeToCreate>& EdgesToCreate, TArray<FEdgeID>& OutNewEdgeIDs )
 {
-	UE_LOG( LogEditableMesh, Verbose, TEXT( "CreateEdges: %s" ), *LogHelpers::ArrayToString( EdgesToCreate ) );
+	EM_ENTER( TEXT( "CreateEdges: %s" ), *LogHelpers::ArrayToString( EdgesToCreate ) );
 
 	// Create new edges in the mesh description
 	{
@@ -4073,7 +4098,7 @@ void UEditableMesh::CreateEdges( const TArray<FEdgeToCreate>& EdgesToCreate, TAr
 		}
 	}
 
-	UE_LOG( LogEditableMesh, Verbose, TEXT( "* CreateEdges returned %s" ), *LogHelpers::ArrayToString( OutNewEdgeIDs ) );
+	EM_EXIT( TEXT( "CreateEdges returned %s" ), *LogHelpers::ArrayToString( OutNewEdgeIDs ) );
 }
 
 
@@ -4219,6 +4244,9 @@ void UEditableMesh::CreatePolygonContour( const TArray<FVertexAndAttributes>& Co
 			FVertexInstanceToCreate& VertexInstanceToCreate = VertexInstancesToCreate.Last();
 			VertexInstanceToCreate.VertexID = VertexID;
 			VertexInstanceToCreate.VertexInstanceAttributes = Contour[ VertexNumber ].PolygonVertexAttributes;
+
+			// Add vertex to list of potential vertices to be merged at the end of the operation
+			VerticesPendingMerging.Add( VertexID );
 		}
 	}
 
@@ -4247,7 +4275,7 @@ void UEditableMesh::CreatePolygonContour( const TArray<FVertexAndAttributes>& Co
 
 void UEditableMesh::CreatePolygons( const TArray<FPolygonToCreate>& PolygonsToCreate, TArray<FPolygonID>& OutNewPolygonIDs, TArray<FEdgeID>& OutNewEdgeIDs )
 {
-	UE_LOG( LogEditableMesh, Verbose, TEXT( "CreatePolygons: %s" ), *LogHelpers::ArrayToString( PolygonsToCreate ) );
+	EM_ENTER( TEXT( "CreatePolygons: %s" ), *LogHelpers::ArrayToString( PolygonsToCreate ) );
 
 	OutNewPolygonIDs.Reset( PolygonsToCreate.Num() );
 	OutNewEdgeIDs.Reset();
@@ -4324,8 +4352,6 @@ void UEditableMesh::CreatePolygons( const TArray<FPolygonToCreate>& PolygonsToCr
 		SetEdgesAttributes( AttributesForEdges );
 	}
 
-	// @todo editable mesh: merge vertex instances if possible here
-
 	for( UEditableMeshAdapter* Adapter : Adapters )
 	{
 		Adapter->OnCreatePolygons( this, OutNewPolygonIDs );
@@ -4359,7 +4385,7 @@ void UEditableMesh::CreatePolygons( const TArray<FPolygonToCreate>& PolygonsToCr
 		AddUndo( MakeUnique<FDeletePolygonsChange>( MoveTemp( RevertInput ) ) );
 	}
 
-	UE_LOG( LogEditableMesh, Verbose, TEXT( "* CreatePolygons returned %s, %s" ), *LogHelpers::ArrayToString( OutNewPolygonIDs ), *LogHelpers::ArrayToString( OutNewEdgeIDs ) );
+	EM_EXIT( TEXT( "CreatePolygons returned %s, %s" ), *LogHelpers::ArrayToString( OutNewPolygonIDs ), *LogHelpers::ArrayToString( OutNewEdgeIDs ) );
 }
 
 
@@ -4384,7 +4410,7 @@ void UEditableMesh::BackupPolygonContour( const FMeshPolygonContour& Contour, TA
 
 void UEditableMesh::DeletePolygons( const TArray<FPolygonID>& PolygonIDsToDelete, const bool bDeleteOrphanedEdges, const bool bDeleteOrphanedVertices, const bool bDeleteOrphanedVertexInstances, const bool bDeleteEmptyPolygonGroups )
 {
-	UE_LOG( LogEditableMesh, Verbose, TEXT( "DeletePolygons: %s" ), *LogHelpers::ArrayToString( PolygonIDsToDelete ) );
+	EM_ENTER( TEXT( "DeletePolygons: %s" ), *LogHelpers::ArrayToString( PolygonIDsToDelete ) );
 
 	FPolygonArray& Polygons = GetMeshDescription()->Polygons();
 	FPolygonGroupArray& PolygonGroups = GetMeshDescription()->PolygonGroups();
@@ -4487,13 +4513,13 @@ void UEditableMesh::DeletePolygons( const TArray<FPolygonID>& PolygonIDsToDelete
 		PolygonsPendingTriangulation.Remove( PolygonID );
 	}
 
-	UE_LOG( LogEditableMesh, Verbose, TEXT( "* DeletePolygons returned" ) );
+	EM_EXIT( TEXT( "DeletePolygons returned" ) );
 }
 
 
 void UEditableMesh::CreatePolygonGroups( const TArray<FPolygonGroupToCreate>& PolygonGroupsToCreate, TArray<FPolygonGroupID>& OutNewPolygonGroupIDs )
 {
-	UE_LOG( LogEditableMesh, Verbose, TEXT( "CreatePolygonGroups: %s" ), *LogHelpers::ArrayToString( PolygonGroupsToCreate ) );
+	EM_ENTER( TEXT( "CreatePolygonGroups: %s" ), *LogHelpers::ArrayToString( PolygonGroupsToCreate ) );
 
 	// Create polygon groups and initialize them
 	{
@@ -4553,13 +4579,13 @@ void UEditableMesh::CreatePolygonGroups( const TArray<FPolygonGroupToCreate>& Po
 		AddUndo( MakeUnique<FDeletePolygonGroupsChange>( MoveTemp( RevertInput ) ) );
 	}
 
-	UE_LOG( LogEditableMesh, Verbose, TEXT( "* CreatePolygonGroups returned %s" ), *LogHelpers::ArrayToString( OutNewPolygonGroupIDs ) );
+	EM_EXIT( TEXT( "CreatePolygonGroups returned %s" ), *LogHelpers::ArrayToString( OutNewPolygonGroupIDs ) );
 }
 
 
 void UEditableMesh::DeletePolygonGroups( const TArray<FPolygonGroupID>& PolygonGroupIDs )
 {
-	UE_LOG( LogEditableMesh, Verbose, TEXT( "DeletePolygonGroups: %s" ), *LogHelpers::ArrayToString( PolygonGroupIDs ) );
+	EM_ENTER( TEXT( "DeletePolygonGroups: %s" ), *LogHelpers::ArrayToString( PolygonGroupIDs ) );
 
 	FPolygonGroupArray& PolygonGroups = GetMeshDescription()->PolygonGroups();
 
@@ -4595,12 +4621,14 @@ void UEditableMesh::DeletePolygonGroups( const TArray<FPolygonGroupID>& PolygonG
 			GetMeshDescription()->DeletePolygonGroup( PolygonGroupID );
 		}
 	}
+
+	EM_EXIT( TEXT( "DeletePolygonGroups returned" ) );
 }
 
 
 void UEditableMesh::SetVerticesAttributes( const TArray<FAttributesForVertex>& AttributesForVertices )
 {
-	UE_LOG( LogEditableMesh, Verbose, TEXT( "SetVerticesAttributes: %s" ), *LogHelpers::ArrayToString( AttributesForVertices ) );
+	EM_ENTER( TEXT( "SetVerticesAttributes: %s" ), *LogHelpers::ArrayToString( AttributesForVertices ) );
 
 	FSetVerticesAttributesChangeInput RevertInput;
 
@@ -4625,13 +4653,13 @@ void UEditableMesh::SetVerticesAttributes( const TArray<FAttributesForVertex>& A
 
 	AddUndo( MakeUnique<FSetVerticesAttributesChange>( MoveTemp( RevertInput ) ) );
 
-	UE_LOG( LogEditableMesh, Verbose, TEXT( "* SetVerticesAttributes returned" ) );
+	EM_EXIT( TEXT( "SetVerticesAttributes returned" ) );
 }
 
 
 void UEditableMesh::SetVertexInstancesAttributes( const TArray<FAttributesForVertexInstance>& AttributesForVertexInstances )
 {
-	UE_LOG( LogEditableMesh, Verbose, TEXT( "SetVertexInstancesAttributes: %s" ), *LogHelpers::ArrayToString( AttributesForVertexInstances ) );
+	EM_ENTER( TEXT( "SetVertexInstancesAttributes: %s" ), *LogHelpers::ArrayToString( AttributesForVertexInstances ) );
 
 	FSetVertexInstancesAttributesChangeInput RevertInput;
 
@@ -4658,13 +4686,13 @@ void UEditableMesh::SetVertexInstancesAttributes( const TArray<FAttributesForVer
 
 	AddUndo( MakeUnique<FSetVertexInstancesAttributesChange>( MoveTemp( RevertInput ) ) );
 
-	UE_LOG( LogEditableMesh, Verbose, TEXT( "* SetVertexInstancesAttributes returned" ) );
+	EM_EXIT( TEXT( "SetVertexInstancesAttributes returned" ) );
 }
 
 
 void UEditableMesh::SetEdgesAttributes( const TArray<FAttributesForEdge>& AttributesForEdges )
 {
-	UE_LOG( LogEditableMesh, Verbose, TEXT( "SetEdgesAttributes: %s" ), *LogHelpers::ArrayToString( AttributesForEdges ) );
+	EM_ENTER( TEXT( "SetEdgesAttributes: %s" ), *LogHelpers::ArrayToString( AttributesForEdges ) );
 
 	FSetEdgesAttributesChangeInput RevertInput;
 
@@ -4689,13 +4717,13 @@ void UEditableMesh::SetEdgesAttributes( const TArray<FAttributesForEdge>& Attrib
 
 	AddUndo( MakeUnique<FSetEdgesAttributesChange>( MoveTemp( RevertInput ) ) );
 
-	UE_LOG( LogEditableMesh, Verbose, TEXT( "* SetEdgesAttributes returned" ) );
+	EM_EXIT( TEXT( "SetEdgesAttributes returned" ) );
 }
 
 
 void UEditableMesh::ChangePolygonsVertexInstances( const TArray<FChangeVertexInstancesForPolygon>& VertexInstancesForPolygons )
 {
-	UE_LOG( LogEditableMesh, Verbose, TEXT( "ChangePolygonsVertexInstances: %s" ), *LogHelpers::ArrayToString( VertexInstancesForPolygons ) );
+	EM_ENTER( TEXT( "ChangePolygonsVertexInstances: %s" ), *LogHelpers::ArrayToString( VertexInstancesForPolygons ) );
 
 	FPolygonArray& Polygons = GetMeshDescription()->Polygons();
 	FVertexInstanceArray& VertexInstances = GetMeshDescription()->VertexInstances();
@@ -4822,7 +4850,7 @@ void UEditableMesh::ChangePolygonsVertexInstances( const TArray<FChangeVertexIns
 		Adapter->OnChangePolygonVertexInstances( this, PolygonIDs );
 	}
 
-	UE_LOG( LogEditableMesh, Verbose, TEXT( "* ChangePolygonsVertexInstances returned" ) );
+	EM_EXIT( TEXT( "ChangePolygonsVertexInstances returned" ) );
 }
 
 
@@ -4860,7 +4888,7 @@ void UEditableMesh::GetConnectedSoftEdges( const FVertexID VertexID, TArray<FEdg
 
 void UEditableMesh::SetPolygonsVertexAttributes( const TArray<FVertexAttributesForPolygon>& VertexAttributesForPolygons )
 {
-	UE_LOG( LogEditableMesh, Verbose, TEXT( "SetPolygonsVertexAttributes: %s" ), *LogHelpers::ArrayToString( VertexAttributesForPolygons ) );
+	EM_ENTER( TEXT( "SetPolygonsVertexAttributes: %s" ), *LogHelpers::ArrayToString( VertexAttributesForPolygons ) );
 
 	FPolygonArray& Polygons = GetMeshDescription()->Polygons();
 
@@ -4877,7 +4905,7 @@ void UEditableMesh::SetPolygonsVertexAttributes( const TArray<FVertexAttributesF
 		}
 	}
 
-	UE_LOG( LogEditableMesh, Verbose, TEXT( "* SetPolygonsVertexAttributes returned" ) );
+	EM_EXIT( TEXT( "SetPolygonsVertexAttributes returned" ) );
 }
 
 
@@ -4962,7 +4990,7 @@ void UEditableMesh::SetPolygonContourVertexAttributes( FMeshPolygonContour& Cont
 
 void UEditableMesh::TryToRemovePolygonEdge( const FEdgeID EdgeID, bool& bOutWasEdgeRemoved, FPolygonID& OutNewPolygonID )
 {
-	UE_LOG( LogEditableMesh, Verbose, TEXT( "TryToRemovePolygonEdge: %s" ), *EdgeID.ToString() );
+	EM_ENTER( TEXT( "TryToRemovePolygonEdge: %s" ), *EdgeID.ToString() );
 
 	bOutWasEdgeRemoved = false;
 	OutNewPolygonID = FPolygonID::Invalid;
@@ -5134,13 +5162,13 @@ void UEditableMesh::TryToRemovePolygonEdge( const FEdgeID EdgeID, bool& bOutWasE
 		}
 	}
 
-	UE_LOG( LogEditableMesh, Verbose, TEXT( "* TryToRemovePolygonEdge returned %s" ), *OutNewPolygonID.ToString() );
+	EM_EXIT( TEXT( "TryToRemovePolygonEdge returned %s" ), *OutNewPolygonID.ToString() );
 }
 
 
 void UEditableMesh::TryToRemoveVertex( const FVertexID VertexID, bool& bOutWasVertexRemoved, FEdgeID& OutNewEdgeID )
 {
-	UE_LOG( LogEditableMesh, Verbose, TEXT( "TryToRemoveVertex: %s" ), *VertexID.ToString() );
+	EM_ENTER( TEXT( "TryToRemoveVertex: %s" ), *VertexID.ToString() );
 
 	bOutWasVertexRemoved = false;
 	OutNewEdgeID = FEdgeID::Invalid;
@@ -5245,13 +5273,13 @@ void UEditableMesh::TryToRemoveVertex( const FVertexID VertexID, bool& bOutWasVe
 		OutNewEdgeID = NewEdgeID;
 	}
 
-	UE_LOG( LogEditableMesh, Verbose, TEXT( "* TryToRemoveVertex returned %s" ), *OutNewEdgeID.ToString() );
+	EM_EXIT( TEXT( "TryToRemoveVertex returned %s" ), *OutNewEdgeID.ToString() );
 }
 
 
 void UEditableMesh::ExtrudePolygons( const TArray<FPolygonID>& PolygonIDs, const float ExtrudeDistance, const bool bKeepNeighborsTogether, TArray<FPolygonID>& OutNewExtrudedFrontPolygons )
 {
-	UE_LOG( LogEditableMesh, Verbose, TEXT( "ExtrudePolygons: %s" ), *LogHelpers::ArrayToString( PolygonIDs ) );
+	EM_ENTER( TEXT( "ExtrudePolygons: %s" ), *LogHelpers::ArrayToString( PolygonIDs ) );
 
 	// @todo mesheditor perf: We can make this much faster by batching up polygons together.  Just be careful about how neighbors are handled.
 	OutNewExtrudedFrontPolygons.Reset();
@@ -5498,6 +5526,7 @@ void UEditableMesh::ExtrudePolygons( const TArray<FPolygonID>& PolygonIDs, const
 							FPolygonToCreate& PolygonToCreate = *new( PolygonsToCreate ) FPolygonToCreate();
 							PolygonToCreate.PolygonGroupID = PolygonGroupID;
 							PolygonToCreate.PerimeterVertices = NewSidePolygonVertices;	// @todo mesheditor perf: Copying static array here, ideally allocations could be avoided
+							PolygonToCreate.PolygonEdgeHardness = EPolygonEdgeHardness::AllEdgesHard;
 
 							// NOTE: We never create holes in side polygons
 
@@ -5510,21 +5539,6 @@ void UEditableMesh::ExtrudePolygons( const TArray<FPolygonID>& PolygonIDs, const
 							NewSidePolygonID = NewPolygonIDs[ 0 ];
 						}
 						AllNewPolygons.Add( NewSidePolygonID );
-
-						// All of this edges of the new polygon will be hard
-						{
-							const int32 NewPerimeterEdgeCount = this->GetPolygonPerimeterEdgeCount( NewSidePolygonID );
-							for( int32 NewPerimeterEdgeNumber = 0; NewPerimeterEdgeNumber < NewPerimeterEdgeCount; ++NewPerimeterEdgeNumber )
-							{
-								bool bNewEdgeWindingIsReversedForPolygon;
-								const FEdgeID NewEdgeID = this->GetPolygonPerimeterEdge( NewSidePolygonID, NewPerimeterEdgeNumber, /* Out */ bNewEdgeWindingIsReversedForPolygon );
-
-								// New side polygons get hard edges, but not hard crease weights
-								FAttributesForEdge& AttributesForEdge = *new( AttributesForEdges ) FAttributesForEdge();
-								AttributesForEdge.EdgeID = NewEdgeID;
-								AttributesForEdge.EdgeAttributes.Attributes.Emplace( MeshAttribute::Edge::IsHard, 0, FMeshElementAttributeValue( true ) );
-							}
-						}
 					}
 				}
 			}
@@ -5632,7 +5646,7 @@ void UEditableMesh::ExtrudePolygons( const TArray<FPolygonID>& PolygonIDs, const
 		DeletePolygons( PolygonIDs, bDeleteOrphanedEdges, bDeleteOrphanedVertices, bDeleteOrphanedVertexInstances, bDeleteEmptyPolygonGroups );
 	}
 
-	UE_LOG( LogEditableMesh, Verbose, TEXT( "* ExtrudePolygons returned %s" ), *LogHelpers::ArrayToString( OutNewExtrudedFrontPolygons ) );
+	EM_EXIT( TEXT( "ExtrudePolygons returned %s" ), *LogHelpers::ArrayToString( OutNewExtrudedFrontPolygons ) );
 }
 
 
@@ -6695,7 +6709,7 @@ void UEditableMesh::SetEdgesHardnessAutomatically( const TArray<FEdgeID>& EdgeID
 
 void UEditableMesh::SetEdgesVertices( const TArray<FVerticesForEdge>& VerticesForEdges )
 {
-	UE_LOG( LogEditableMesh, Verbose, TEXT( "SetEdgesVertices: %s" ), *LogHelpers::ArrayToString( VerticesForEdges ) );
+	EM_ENTER( TEXT( "SetEdgesVertices: %s" ), *LogHelpers::ArrayToString( VerticesForEdges ) );
 
 	FEdgeArray& Edges = GetMeshDescription()->Edges();
 	FVertexArray& Vertices = GetMeshDescription()->Vertices();
@@ -6749,13 +6763,13 @@ void UEditableMesh::SetEdgesVertices( const TArray<FVerticesForEdge>& VerticesFo
 
 	AddUndo( MakeUnique<FSetEdgesVerticesChange>( MoveTemp( RevertInput ) ) );
 
-	UE_LOG( LogEditableMesh, Verbose, TEXT( "* SetEdgesVertices finished") );
+	EM_EXIT( TEXT( "SetEdgesVertices finished") );
 }
 
 
 void UEditableMesh::InsertPolygonPerimeterVertices( const FPolygonID PolygonID, const int32 InsertBeforeVertexNumber, const TArray<FVertexAndAttributes>& VerticesToInsert )
 {
-	UE_LOG( LogEditableMesh, Verbose, TEXT( "InsertPolygonPerimeterVertices: PolygonID:%s InsertBeforeVertexNumber:%d %s" ), *PolygonID.ToString(), InsertBeforeVertexNumber, *LogHelpers::ArrayToString( VerticesToInsert ) );
+	EM_ENTER( TEXT( "InsertPolygonPerimeterVertices: PolygonID:%s InsertBeforeVertexNumber:%d %s" ), *PolygonID.ToString(), InsertBeforeVertexNumber, *LogHelpers::ArrayToString( VerticesToInsert ) );
 	// @todo mesheditor: see if we want to keep this action. It is only used by SplitEdge, and doesn't generate missing edges.
 	// We can achieve the same thing by deleting a polygon without deleting vertex instances, and creating a new polygon with an extra vertex instance.
 
@@ -6806,13 +6820,13 @@ void UEditableMesh::InsertPolygonPerimeterVertices( const FPolygonID PolygonID,
 		}
 	}
 
-	UE_LOG( LogEditableMesh, Verbose, TEXT( "* InsertPolygonPerimeterVertices finished" ) );
+	EM_EXIT( TEXT( "InsertPolygonPerimeterVertices finished" ) );
 }
 
 
 void UEditableMesh::RemovePolygonPerimeterVertices( const FPolygonID PolygonID, const int32 FirstVertexNumberToRemove, const int32 NumVerticesToRemove, const bool bDeleteOrphanedVertexInstances )
 {
-	UE_LOG( LogEditableMesh, Verbose, TEXT( "RemovePolygonPerimeterVertices: PolygonID:%s FirstVertexNumberToRemove:%d NumVerticesToRemove:%d" ), *PolygonID.ToString(), FirstVertexNumberToRemove, NumVerticesToRemove );
+	EM_ENTER( TEXT( "RemovePolygonPerimeterVertices: PolygonID:%s FirstVertexNumberToRemove:%d NumVerticesToRemove:%d" ), *PolygonID.ToString(), FirstVertexNumberToRemove, NumVerticesToRemove );
 	FMeshPolygon& Polygon = GetMeshDescription()->GetPolygon( PolygonID );
 
 	// Back up 
@@ -6880,7 +6894,7 @@ void UEditableMesh::RemovePolygonPerimeterVertices( const FPolygonID PolygonID,
 		}
 	}
 
-	UE_LOG( LogEditableMesh, Verbose, TEXT( "* RemovePolygonPerimeterVertices finished" ) );
+	EM_EXIT( TEXT( "RemovePolygonPerimeterVertices finished" ) );
 }
 
 
diff --git a/Engine/Plugins/Runtime/EditableMesh/Source/EditableMesh/Public/EditableMesh.h b/Engine/Plugins/Runtime/EditableMesh/Source/EditableMesh/Public/EditableMesh.h
--- a/Engine/Plugins/Runtime/EditableMesh/Source/EditableMesh/Public/EditableMesh.h
+++ b/Engine/Plugins/Runtime/EditableMesh/Source/EditableMesh/Public/EditableMesh.h
@@ -15,6 +15,7 @@
 
 class UEditableMeshAdapter;
 
+#define EDITABLE_MESH_DEBUG_LOG
 #ifdef EDITABLE_MESH_DEBUG_LOG
 EDITABLEMESH_API DECLARE_LOG_CATEGORY_EXTERN( LogEditableMesh, All, All );
 #else
