Fixed variadic macro so code compiles in GCC when passed zero parameters.

diff --git a/Engine/Plugins/Runtime/EditableMesh/Source/EditableMesh/EditableMesh.cpp b/Engine/Plugins/Runtime/EditableMesh/Source/EditableMesh/EditableMesh.cpp
--- a/Engine/Plugins/Runtime/EditableMesh/Source/EditableMesh/EditableMesh.cpp
+++ b/Engine/Plugins/Runtime/EditableMesh/Source/EditableMesh/EditableMesh.cpp
@@ -52,8 +52,8 @@ static int32 GIndent = 0;
 static FString IncreaseIndent() { GIndent += 2; return FString::ChrN( GIndent - 2, ' ' ); }
 static FString DecreaseIndent() { GIndent -= 2; return FString::ChrN( GIndent, ' ' ); }
 
-#define EM_ENTER(Text, ...) UE_LOG( LogEditableMesh, Verbose, TEXT( "%s" ) Text, *IncreaseIndent(), __VA_ARGS__ )
-#define EM_EXIT(Text, ...) UE_LOG( LogEditableMesh, Verbose, TEXT( "%s" ) Text, *DecreaseIndent(), __VA_ARGS__ )
+#define EM_ENTER(Text, ...) UE_LOG( LogEditableMesh, Verbose, TEXT( "%s" ) Text, *IncreaseIndent(), ##__VA_ARGS__ )
+#define EM_EXIT(Text, ...) UE_LOG( LogEditableMesh, Verbose, TEXT( "%s" ) Text, *DecreaseIndent(), ##__VA_ARGS__ )
 #else
 #define EM_ENTER(Text, ...)
 #define EM_EXIT(Text, ...)
