Added more detailed debug logging for editable mesh actions.

diff --git a/Engine/Plugins/Runtime/EditableMesh/Source/EditableMesh/EditableMesh.cpp b/Engine/Plugins/Runtime/EditableMesh/Source/EditableMesh/EditableMesh.cpp
--- a/Engine/Plugins/Runtime/EditableMesh/Source/EditableMesh/EditableMesh.cpp
+++ b/Engine/Plugins/Runtime/EditableMesh/Source/EditableMesh/EditableMesh.cpp
@@ -2893,7 +2893,7 @@ void UEditableMesh::CreateMissingPolygonHoleEdges( const FPolygonID PolygonID, c
 
 void UEditableMesh::SplitEdge( const FEdgeID EdgeID, const TArray<float>& Splits, TArray<FVertexID>& OutNewVertexIDs )
 {
-	EM_ENTER( TEXT( "SplitEdge: %s" ), *EdgeID.ToString() );
+	EM_ENTER( TEXT( "SplitEdge: %s %s" ), *EdgeID.ToString(), *LogHelpers::ArrayToString( Splits ) );
 
 	// NOTE: The incoming splits should always be between 0.0 and 1.0, representing progress along 
 	//       the edge from the edge's first vertex toward it's other vertex.  The order doesn't matter (we'll sort them.)
@@ -3397,6 +3397,8 @@ void UEditableMesh::FindPolygonLoop( const FEdgeID EdgeID, TArray<FEdgeID>& OutE
 
 void UEditableMesh::InsertEdgeLoop( const FEdgeID EdgeID, const TArray<float>& Splits, TArray<FEdgeID>& OutNewEdgeIDs )
 {
+	EM_ENTER( TEXT( "InsertEdgeLoop: %s %s" ), *EdgeID.ToString(), *LogHelpers::ArrayToString( Splits ) );
+
 	OutNewEdgeIDs.Reset();
 
 	// NOTE: The incoming splits should always be between 0.0 and 1.0, representing progress along 
@@ -3528,11 +3530,15 @@ void UEditableMesh::InsertEdgeLoop( const FEdgeID EdgeID, const TArray<float>& S
 			OutNewEdgeIDs.Append( NewEdgeIDs );
 		}
 	}
+
+	EM_EXIT( TEXT( "InsertEdgeLoop returned %s" ), *LogHelpers::ArrayToString( OutNewEdgeIDs ) );
 }
 
 
 void UEditableMesh::SplitPolygons( const TArray<FPolygonToSplit>& PolygonsToSplit, TArray<FEdgeID>& OutNewEdgeIDs )
 {
+	EM_ENTER( TEXT( "SplitPolygons: %s" ), *LogHelpers::ArrayToString( PolygonsToSplit ) );
+
 	OutNewEdgeIDs.Reset();
 
 	static TArray<FPolygonToCreate> PolygonsToCreate;
@@ -3670,11 +3676,15 @@ void UEditableMesh::SplitPolygons( const TArray<FPolygonToSplit>& PolygonsToSpli
 	CreatePolygons( PolygonsToCreate, /* Out */ NewPolygonIDs, /* Out */ NewEdgeIDs );
 
 	OutNewEdgeIDs.Append( NewEdgeIDs );
+
+	EM_EXIT( TEXT( "SplitPolygons returned %s" ), *LogHelpers::ArrayToString( OutNewEdgeIDs ) );
 }
 
 
 void UEditableMesh::DeleteEdgeAndConnectedPolygons( const FEdgeID EdgeID, const bool bDeleteOrphanedEdges, const bool bDeleteOrphanedVertices, const bool bDeleteOrphanedVertexInstances, const bool bDeleteEmptyPolygonGroups )
 {
+	EM_ENTER( TEXT( "DeleteEdgeAndConnectedPolygons: %s %s %s %s %s" ), *EdgeID.ToString(), *LogHelpers::BoolToString( bDeleteOrphanedEdges ), *LogHelpers::BoolToString( bDeleteOrphanedVertices ), *LogHelpers::BoolToString( bDeleteOrphanedVertexInstances ), *LogHelpers::BoolToString( bDeleteEmptyPolygonGroups ) );
+
 	static TArray<FPolygonID> PolygonIDsToDelete;
 	PolygonIDsToDelete.Reset();
 
@@ -3704,11 +3714,15 @@ void UEditableMesh::DeleteEdgeAndConnectedPolygons( const FEdgeID EdgeID, const
 
 		DeleteEdges( EdgeIDsToDelete, bDeleteOrphanedVertices );
 	}
+
+	EM_EXIT( TEXT( "DeleteEdgeAndConnectedPolygons returned" ) );
 }
 
 
 void UEditableMesh::DeleteVertexAndConnectedEdgesAndPolygons( const FVertexID VertexID, const bool bDeleteOrphanedEdges, const bool bDeleteOrphanedVertices, const bool bDeleteOrphanedVertexInstances, const bool bDeleteEmptyPolygonGroups )
 {
+	EM_ENTER( TEXT( "DeleteVertexAndConnectedEdgesAndPolygons: %s %s %s %s %s" ), *VertexID.ToString(), *LogHelpers::BoolToString( bDeleteOrphanedEdges ), *LogHelpers::BoolToString( bDeleteOrphanedVertices ), *LogHelpers::BoolToString( bDeleteOrphanedVertexInstances ), *LogHelpers::BoolToString( bDeleteEmptyPolygonGroups ) );
+
 	static TArray<FEdgeID> EdgeIDsToDelete;
 	EdgeIDsToDelete.Reset();
 
@@ -3728,6 +3742,8 @@ void UEditableMesh::DeleteVertexAndConnectedEdgesAndPolygons( const FVertexID Ve
 			DeleteEdgeAndConnectedPolygons( EdgeIDToDelete, bDeleteOrphanedEdges, bDeleteOrphanedVertices, bDeleteOrphanedVertexInstances, bDeleteEmptyPolygonGroups );
 		}
 	}
+
+	EM_EXIT( TEXT( "DeleteVertexAndConnectedEdgesAndPolygons returned" ) );
 }
 
 
@@ -3840,7 +3856,7 @@ void UEditableMesh::DeleteVertexInstances( const TArray<FVertexInstanceID>& Vert
 
 void UEditableMesh::DeleteEdges( const TArray<FEdgeID>& EdgeIDsToDelete, const bool bDeleteOrphanedVertices )
 {
-	EM_ENTER( TEXT( "DeleteEdges: %s" ), *LogHelpers::ArrayToString( EdgeIDsToDelete ) );
+	EM_ENTER( TEXT( "DeleteEdges: %s %s" ), *LogHelpers::ArrayToString( EdgeIDsToDelete ), *LogHelpers::BoolToString( bDeleteOrphanedVertices ) );
 
 	FEdgeArray& Edges = GetMeshDescription()->Edges();
 	FVertexArray& Vertices = GetMeshDescription()->Vertices();
@@ -5162,7 +5178,7 @@ void UEditableMesh::TryToRemovePolygonEdge( const FEdgeID EdgeID, bool& bOutWasE
 		}
 	}
 
-	EM_EXIT( TEXT( "TryToRemovePolygonEdge returned %s" ), *OutNewPolygonID.ToString() );
+	EM_EXIT( TEXT( "TryToRemovePolygonEdge returned %s %s" ), *LogHelpers::BoolToString( bOutWasEdgeRemoved ), *OutNewPolygonID.ToString() );
 }
 
 
@@ -5273,7 +5289,7 @@ void UEditableMesh::TryToRemoveVertex( const FVertexID VertexID, bool& bOutWasVe
 		OutNewEdgeID = NewEdgeID;
 	}
 
-	EM_EXIT( TEXT( "TryToRemoveVertex returned %s" ), *OutNewEdgeID.ToString() );
+	EM_EXIT( TEXT( "TryToRemoveVertex returned %s %s" ), *LogHelpers::BoolToString( bOutWasVertexRemoved ), *OutNewEdgeID.ToString() );
 }
 
 
@@ -5652,6 +5668,8 @@ void UEditableMesh::ExtrudePolygons( const TArray<FPolygonID>& PolygonIDs, const
 
 void UEditableMesh::ExtendEdges( const TArray<FEdgeID>& EdgeIDs, const bool bWeldNeighbors, TArray<FEdgeID>& OutNewExtendedEdgeIDs )
 {
+	EM_ENTER( TEXT( "ExtendEdges: %s" ), *LogHelpers::ArrayToString( EdgeIDs ) );
+
 	OutNewExtendedEdgeIDs.Reset();
 
 	static TArray<FVertexID> NewVertexIDs;
@@ -5795,11 +5813,15 @@ void UEditableMesh::ExtendEdges( const TArray<FEdgeID>& EdgeIDs, const bool bWel
 		// to be created, depending on how many edges we share with neighbors that were extended.
 		check( bWeldNeighbors ? ( NewEdgeIDs.Num() <= 2 * EdgeIDs.Num() ) : ( NewEdgeIDs.Num() == 2 * EdgeIDs.Num() ) );
 	}
+
+	EM_EXIT( TEXT( "ExtendEdges returned %s" ), *LogHelpers::ArrayToString( OutNewExtendedEdgeIDs ) );
 }
 
 
 void UEditableMesh::ExtendVertices( const TArray<FVertexID>& VertexIDs, const bool bOnlyExtendClosestEdge, const FVector ReferencePosition, TArray<FVertexID>& OutNewExtendedVertexIDs )
 {
+	EM_ENTER( TEXT( "ExtendVertices: %s" ), *LogHelpers::ArrayToString( VertexIDs ) );
+
 	OutNewExtendedVertexIDs.Reset();
 
 	// Create new vertices for all of the new edges.  If bWeldNeighbors is true, we'll share vertices between edges that share the
@@ -5965,6 +5987,8 @@ void UEditableMesh::ExtendVertices( const TArray<FVertexID>& VertexIDs, const bo
 		NewEdgeIDs.Reset();
 		CreatePolygons( PolygonsToCreate, /* Out */ NewPolygonIDs, /* Out */ NewEdgeIDs );
 	}
+
+	EM_EXIT( TEXT( "ExtendVertices returned %s" ), *LogHelpers::ArrayToString( OutNewExtendedVertexIDs ) );
 }
 
 
@@ -6251,15 +6275,23 @@ void UEditableMesh::BevelOrInsetPolygons( const TArray<FPolygonID>& PolygonIDs,
 
 void UEditableMesh::InsetPolygons( const TArray<FPolygonID>& PolygonIDs, const float InsetFixedDistance, const float InsetProgressTowardCenter, const EInsetPolygonsMode Mode, TArray<FPolygonID>& OutNewCenterPolygonIDs, TArray<FPolygonID>& OutNewSidePolygonIDs )
 {
+	EM_ENTER( TEXT( "InsetPolygons: %s %f %f" ), *LogHelpers::ArrayToString( PolygonIDs ), InsetFixedDistance, InsetProgressTowardCenter );
+
 	const bool bShouldBevel = false;
 	BevelOrInsetPolygons( PolygonIDs, InsetFixedDistance, InsetProgressTowardCenter, Mode, bShouldBevel, /* Out */ OutNewCenterPolygonIDs, /* Out */ OutNewSidePolygonIDs );
+
+	EM_EXIT( TEXT( "InsetPolygons returned %s %s" ), *LogHelpers::ArrayToString( OutNewCenterPolygonIDs ), *LogHelpers::ArrayToString( OutNewSidePolygonIDs ) );
 }
 
 
 void UEditableMesh::BevelPolygons( const TArray<FPolygonID>& PolygonIDs, const float BevelFixedDistance, const float BevelProgressTowardCenter, TArray<FPolygonID>& OutNewCenterPolygonIDs, TArray<FPolygonID>& OutNewSidePolygonIDs )
 {
+	EM_ENTER( TEXT( "BevelPolygons: %s %f %f" ), *LogHelpers::ArrayToString( PolygonIDs ), BevelFixedDistance, BevelProgressTowardCenter );
+
 	const bool bShouldBevel = true;
 	BevelOrInsetPolygons( PolygonIDs, BevelFixedDistance, BevelProgressTowardCenter, EInsetPolygonsMode::All, bShouldBevel, /* Out */ OutNewCenterPolygonIDs, /* Out */ OutNewSidePolygonIDs );
+
+	EM_EXIT( TEXT( "BevelPolygons returned %s %s" ), *LogHelpers::ArrayToString( OutNewCenterPolygonIDs ), *LogHelpers::ArrayToString( OutNewSidePolygonIDs ) );
 }
 
 
@@ -6992,6 +7024,8 @@ int32 UEditableMesh::FindPolygonHoleEdgeNumberForVertices( const FPolygonID Poly
 
 void UEditableMesh::FlipPolygons( const TArray<FPolygonID>& PolygonIDs )
 {
+	EM_ENTER( TEXT( "FlipPolygons %s" ), *LogHelpers::ArrayToString( PolygonIDs ) );
+
 	static TArray<FPolygonToCreate> PolygonsToCreate;
 	PolygonsToCreate.Reset();
 	for( int32 PolygonNumber = 0; PolygonNumber < PolygonIDs.Num(); ++PolygonNumber )
@@ -7038,11 +7072,15 @@ void UEditableMesh::FlipPolygons( const TArray<FPolygonID>& PolygonIDs )
 		static TArray<FEdgeID> NewEdgeIDs;
 		CreatePolygons( PolygonsToCreate, NewPolygonIDs, NewEdgeIDs );
 	}
+
+	EM_EXIT( TEXT( "FlipPolygons returned" ) );
 }
 
 
 void UEditableMesh::TriangulatePolygons( const TArray<FPolygonID>& PolygonIDs, TArray<FPolygonID>& OutNewTrianglePolygons )
 {
+	EM_ENTER( TEXT( "TriangulatePolygons: %s" ), *LogHelpers::ArrayToString( PolygonIDs ) );
+
 	OutNewTrianglePolygons.Reset();
 
 	static TArray<FPolygonToCreate> PolygonsToCreate;
@@ -7103,11 +7141,14 @@ void UEditableMesh::TriangulatePolygons( const TArray<FPolygonID>& PolygonIDs, T
 		static TArray<FEdgeID> NewEdgeIDs;
 		CreatePolygons( PolygonsToCreate, /* Out */ OutNewTrianglePolygons, /* Out */ NewEdgeIDs );
 	}
+
+	EM_EXIT( TEXT( "TriangulatePolygons returned %s" ), *LogHelpers::ArrayToString( OutNewTrianglePolygons ) );
 }
 
 
 void UEditableMesh::AssignPolygonsToPolygonGroups( const TArray<FPolygonGroupForPolygon>& PolygonGroupForPolygons, const bool bDeleteOrphanedPolygonGroups )
 {
+	EM_ENTER( TEXT( "AssignPolygonsToPolygonGroups: %s" ), *LogHelpers::ArrayToString( PolygonGroupForPolygons ) );
 	// Back up
 	{
 		FAssignPolygonsToPolygonGroupChangeInput RevertInput;
@@ -7168,11 +7209,15 @@ void UEditableMesh::AssignPolygonsToPolygonGroups( const TArray<FPolygonGroupFor
 	{
 		DeletePolygonGroups( PolygonGroupsToDelete );
 	}
+
+	EM_EXIT( TEXT( "AssignPolygonsToPolygonGroups returned" ) );
 }
 
 
 void UEditableMesh::WeldVertices( const TArray<FVertexID>& VertexIDsToWeld, FVertexID& OutNewVertexID )
 {
+	EM_ENTER( TEXT( "WeldVertices: %s" ), *LogHelpers::ArrayToString( VertexIDsToWeld ) );
+
 	OutNewVertexID = FVertexID::Invalid;
 
 	// This function takes a list of perimeter vertices and a list of vertices to be welded as input.
@@ -7278,6 +7323,7 @@ void UEditableMesh::WeldVertices( const TArray<FVertexID>& VertexIDsToWeld, FVer
 		if( !bIsValid )
 		{
 			// Return with the NewVertexID set to Invalid
+			EM_EXIT( TEXT( "WeldVertices failed" ) );
 			return;
 		}
 	}
@@ -7287,6 +7333,7 @@ void UEditableMesh::WeldVertices( const TArray<FVertexID>& VertexIDsToWeld, FVer
 		// For now, abort if we don't need to create a welded vertex.
 		// This generally implies that all (or a disconnected subset) of the mesh is about to disappear,
 		// which arguably is not something we would want to do like this anyway.
+		EM_EXIT( TEXT( "WeldVertices failed" ) );
 		return;
 	}
 
@@ -7432,11 +7479,15 @@ void UEditableMesh::WeldVertices( const TArray<FVertexID>& VertexIDsToWeld, FVer
 		const bool bDeleteEmptyPolygonGroups = false;
 		DeletePolygons( AllConnectedPolygonIDs, bDeleteOrphanedEdges, bDeleteOrphanedVertices, bDeleteOrphanedVertexInstances, bDeleteEmptyPolygonGroups );
 	}
+
+	EM_EXIT( TEXT( "WeldVertices returned %s" ), *OutNewVertexID.ToString() );
 }
 
 
 void UEditableMesh::TessellatePolygons( const TArray<FPolygonID>& PolygonIDs, const ETriangleTessellationMode TriangleTessellationMode, TArray<FPolygonID>& OutNewPolygonIDs )
 {
+	EM_ENTER( TEXT( "TessellatePolygons: %s" ), *LogHelpers::ArrayToString( PolygonIDs ) );
+
 	OutNewPolygonIDs.Reset();
 
 	const TVertexAttributeArray<FVector>& VertexPositions = GetMeshDescription()->VertexAttributes().GetAttributes<FVector>( MeshAttribute::Vertex::Position );
@@ -7703,6 +7754,8 @@ void UEditableMesh::TessellatePolygons( const TArray<FPolygonID>& PolygonIDs, co
 		static TArray<FEdgeID> NewEdgeIDs;
 		CreatePolygons( PolygonsToCreate, /* Out */ OutNewPolygonIDs, /* Out */ NewEdgeIDs );
 	}
+
+	EM_EXIT( TEXT( "TessellatePolygons returned %s" ), *LogHelpers::ArrayToString( OutNewPolygonIDs ) );
 }
 
 
