Made LogHelper macros work for any array of an arithmetic native type, without requiring explicit specializations.
Added missing ToString() methods.

diff --git a/Engine/Plugins/Runtime/EditableMesh/Source/EditableMesh/Public/EditableMeshTypes.h b/Engine/Plugins/Runtime/EditableMesh/Source/EditableMesh/Public/EditableMeshTypes.h
--- a/Engine/Plugins/Runtime/EditableMesh/Source/EditableMesh/Public/EditableMeshTypes.h
+++ b/Engine/Plugins/Runtime/EditableMesh/Source/EditableMesh/Public/EditableMeshTypes.h
@@ -14,31 +14,20 @@
 // @todo mesheditor: Move elsewhere
 namespace LogHelpers
 {
-	template<typename ArrayType>
-	inline FString ArrayToString( const TArray<ArrayType>& Array )
+	template <typename ElementType, typename TEnableIf<!TIsArithmetic<ElementType>::Value, int>::Type = 0>
+	inline FString ElementToString( const ElementType& Element )
 	{
-		FString String;
-		if( Array.Num() == 0 )
-		{
-			String = TEXT( "Empty" );
-		}
-		else
-		{
-			for( const ArrayType& Element : Array )
-			{
-				if( !String.IsEmpty() )
-				{
-					String.Append( TEXT( ", " ) );
-				}
-				String.Append( Element.ToString() );
-			}
-			String = TEXT( "[" ) + String + TEXT( "]" );
-		}
-		return String;
+		return Element.ToString();
 	}
 
-	template<>
-	inline FString ArrayToString<int32>( const TArray<int32>& Array )
+	template <typename ElementType, typename TEnableIf<TIsArithmetic<ElementType>::Value, int>::Type = 0>
+	inline FString ElementToString( const ElementType& Element )
+	{
+		return Lex::ToString( Element );
+	}
+
+	template <typename ArrayType>
+	inline FString ArrayToString( const TArray<ArrayType>& Array )
 	{
 		FString String;
 		if( Array.Num() == 0 )
@@ -47,20 +36,19 @@ namespace LogHelpers
 		}
 		else
 		{
-			for( const int32& Element : Array )
+			for( const ArrayType& Element : Array )
 			{
 				if( !String.IsEmpty() )
 				{
 					String.Append( TEXT( ", " ) );
 				}
-				String.Append( FString::FormatAsNumber( Element ) );
+				String.Append( ElementToString( Element ) );
 			}
 			String = TEXT( "[" ) + String + TEXT( "]" );
 		}
 		return String;
 	}
 
-
 	template<typename ArrayType>
 	inline FString ArrayToString( const TArray<TArray<ArrayType>>& Array )
 	{
@@ -798,7 +786,13 @@ struct FVertexPair
 	{
 	}
 
-	// @todo mesheditor: Do we need a ToString() function?  Would it ever be called?
+	FString ToString() const
+	{
+		return FString::Printf(
+			TEXT( "VertexID0:%s, VertexID1:%s" ),
+			*VertexID0.ToString(),
+			*VertexID1.ToString() );
+	}
 };
 
 
@@ -823,7 +817,13 @@ struct FPolygonToSplit
 	{
 	}
 
-	// @todo mesheditor: Do we need a ToString() function?  Would it ever be called?
+	FString ToString() const
+	{
+		return FString::Printf(
+			TEXT( "PolygonID:%s, VertexPairsToSplitAt:%s" ),
+			*PolygonID.ToString(),
+			*LogHelpers::ArrayToString( VertexPairsToSplitAt ) );
+	}
 };
 
 
