- UMeshDescription constructor no longer registers a set of default attributes. Attributes now need to be registered by the mesh description owner.
- Added a static UStaticMesh::RegisterMeshAttributes method to register attributes required by a static mesh.
- Added an additional attribute used by EditableMesh static meshes (material asset name for a polygon group).
- EditableStaticMeshAdapter internals now hold a copy of the static mesh material index they reference. This fixes an undo/redo bug.
- Removed FindOrAddMaterial from the adapter interface as it is no longer required; instead we hold the material asset name as a polygon group attribute and add/remove the material from the static mesh in the adapter.
- Fixed a bug when creating a polygon group when subdivided mesh rendering is enabled.
- Bumped the static mesh DDC key version.

diff --git a/Engine/Plugins/Editor/MeshEditor/Source/MeshEditor/MeshEditorMode.cpp b/Engine/Plugins/Editor/MeshEditor/Source/MeshEditor/MeshEditorMode.cpp
--- a/Engine/Plugins/Editor/MeshEditor/Source/MeshEditor/MeshEditorMode.cpp
+++ b/Engine/Plugins/Editor/MeshEditor/Source/MeshEditor/MeshEditorMode.cpp
@@ -2330,11 +2330,6 @@ bool FMeshEditorMode::AssignMaterialToSelectedPolygons( UMaterialInterface* Sele
 
 			EditableMesh->StartModification( EMeshModificationType::Final, EMeshTopologyChange::TopologyChange );
 			{
-				int32 MaterialIndex = INDEX_NONE;
-				FName ImportedMaterialName = NAME_None;
-				//We should be able to add a material at least one adapter must succeed
-				check(EditableMesh->FindOrAddMaterial(SelectedMaterial, MaterialIndex, ImportedMaterialName));
-
 				const UMeshDescription* MeshDescription = EditableMesh->GetMeshDescription();
 
 				UPrimitiveComponent* Component = nullptr;
@@ -2351,12 +2346,12 @@ bool FMeshEditorMode::AssignMaterialToSelectedPolygons( UMaterialInterface* Sele
 				// further details if there is more than one polygon group which matches the material.
 				FPolygonGroupID PolygonGroupToAssign = FPolygonGroupID::Invalid;
 
-				const TPolygonGroupAttributeArray<FName>& PolygonGroupImportedMaterialSlotName = MeshDescription->PolygonGroupAttributes().GetAttributes<FName>(MeshAttribute::PolygonGroup::ImportedMaterialSlotName);
+				const TPolygonGroupAttributeArray<FName>& PolygonGroupMaterialAssetNames = MeshDescription->PolygonGroupAttributes().GetAttributes<FName>( MeshAttribute::PolygonGroup::MaterialAssetName );
 
 				for( const FPolygonGroupID PolygonGroupID : MeshDescription->PolygonGroups().GetElementIDs() )
 				{
-					FName PolygonGroupMaterialName = PolygonGroupImportedMaterialSlotName[PolygonGroupID];
-					if(ImportedMaterialName == PolygonGroupMaterialName)
+					const FName PolygonGroupMaterialName = PolygonGroupMaterialAssetNames[ PolygonGroupID ];
+					if( SelectedMaterial->GetPathName() == PolygonGroupMaterialName.ToString() )
 					{
 						// We only expect to find one polygon group containing this material at the moment.
 						// We need to provide a way of distinguishing different polygon groups with the same material.
@@ -2368,14 +2363,33 @@ bool FMeshEditorMode::AssignMaterialToSelectedPolygons( UMaterialInterface* Sele
 				// If we didn't find the material being used anywhere, create a new polygon group
 				if( PolygonGroupToAssign == FPolygonGroupID::Invalid )
 				{
+					// Helper function which returns a unique FName for the material slot name, based on the material's asset name,
+					// and adding a unique suffix if there are other polygon groups with the same material slot name.
+					auto MakeUniqueSlotName = [ MeshDescription ]( FName Name ) -> FName
+					{
+						const TPolygonGroupAttributeArray<FName>& MaterialSlotNames = MeshDescription->PolygonGroupAttributes().GetAttributes<FName>( MeshAttribute::PolygonGroup::MaterialAssetName );
+						for( const FPolygonGroupID PolygonGroupID : MeshDescription->PolygonGroups().GetElementIDs() )
+						{
+							const FName ExistingName = MaterialSlotNames[ PolygonGroupID ];
+							if( ExistingName.GetComparisonIndex() == Name.GetComparisonIndex() )
+							{
+								Name = FName( Name, FMath::Max( Name.GetNumber(), ExistingName.GetNumber() + 1 ) );
+							}
+						}
+						return Name;
+					};
+
+					const FName UniqueSlotName = MakeUniqueSlotName( SelectedMaterial->GetFName() );
+
 					static TArray<FPolygonGroupToCreate> PolygonGroupsToCreate;
 					PolygonGroupsToCreate.Reset( 1 );
 					PolygonGroupsToCreate.Emplace();
 					
 					FPolygonGroupToCreate& PolygonGroupToCreate = PolygonGroupsToCreate.Last();
-					PolygonGroupToCreate.PolygonGroupAttributes.Attributes.Emplace(MeshAttribute::PolygonGroup::ImportedMaterialSlotName, 0, FMeshElementAttributeValue(ImportedMaterialName));
+					PolygonGroupToCreate.PolygonGroupAttributes.Attributes.Emplace( MeshAttribute::PolygonGroup::MaterialAssetName, 0, FMeshElementAttributeValue( FName( *SelectedMaterial->GetPathName() ) ) );
+					PolygonGroupToCreate.PolygonGroupAttributes.Attributes.Emplace( MeshAttribute::PolygonGroup::ImportedMaterialSlotName, 0, FMeshElementAttributeValue( UniqueSlotName ));
 					static TArray<FPolygonGroupID> NewPolygonGroupIDs;
-					EditableMesh->CreatePolygonGroups( PolygonGroupsToCreate, NewPolygonGroupIDs);
+					EditableMesh->CreatePolygonGroups( PolygonGroupsToCreate, NewPolygonGroupIDs );
 					PolygonGroupToAssign = NewPolygonGroupIDs[ 0 ];
 				}
 
diff --git a/Engine/Plugins/Editor/MeshEditor/Source/MeshEditor/MeshEditorStaticMeshAdapter.h b/Engine/Plugins/Editor/MeshEditor/Source/MeshEditor/MeshEditorStaticMeshAdapter.h
--- a/Engine/Plugins/Editor/MeshEditor/Source/MeshEditor/MeshEditorStaticMeshAdapter.h
+++ b/Engine/Plugins/Editor/MeshEditor/Source/MeshEditor/MeshEditorStaticMeshAdapter.h
@@ -22,8 +22,6 @@ public:
 	/** Creates a editable static mesh from the specified component and sub-mesh address */
 	void Initialize( UEditableMesh* EditableMesh, class UWireframeMesh* WireframeMesh );
 
-	virtual bool FindOrAddMaterial(class UMaterialInterface* Material, int32& MaterialIndex, FName& ImportedMaterialName) override { return false; }
-
 	virtual void InitializeFromEditableMesh( const UEditableMesh* EditableMesh ) override {}
 	virtual void OnRebuildRenderMeshStart( const UEditableMesh* EditableMesh, const bool bInvalidateLighting ) override;
 	virtual void OnRebuildRenderMesh( const UEditableMesh* EditableMesh ) override;
diff --git a/Engine/Plugins/Editor/MeshEditor/Source/MeshEditor/MeshEditorSubdividedStaticMeshAdapter.h b/Engine/Plugins/Editor/MeshEditor/Source/MeshEditor/MeshEditorSubdividedStaticMeshAdapter.h
--- a/Engine/Plugins/Editor/MeshEditor/Source/MeshEditor/MeshEditorSubdividedStaticMeshAdapter.h
+++ b/Engine/Plugins/Editor/MeshEditor/Source/MeshEditor/MeshEditorSubdividedStaticMeshAdapter.h
@@ -22,8 +22,6 @@ public:
 	/** Creates a editable static mesh from the specified component and sub-mesh address */
 	void Initialize( UEditableMesh* EditableMesh, class UWireframeMesh* WireframeMesh );
 
-	virtual bool FindOrAddMaterial(class UMaterialInterface* Material, int32& MaterialIndex, FName& ImportedMaterialName) override { return false; }
-
 	virtual void InitializeFromEditableMesh( const UEditableMesh* EditableMesh ) override {}
 	virtual void OnRebuildRenderMeshStart( const UEditableMesh* EditableMesh, const bool bInvalidateLighting ) override;
 	virtual void OnRebuildRenderMesh( const UEditableMesh* EditableMesh ) override;
diff --git a/Engine/Plugins/Runtime/EditableMesh/Source/EditableMesh/EditableMesh.cpp b/Engine/Plugins/Runtime/EditableMesh/Source/EditableMesh/EditableMesh.cpp
--- a/Engine/Plugins/Runtime/EditableMesh/Source/EditableMesh/EditableMesh.cpp
+++ b/Engine/Plugins/Runtime/EditableMesh/Source/EditableMesh/EditableMesh.cpp
@@ -71,11 +71,16 @@ namespace EditableMesh
 	static FAutoConsoleVariable UseBoundlessOctree( TEXT( "EditableMesh.UseBoundlessOctree" ), 1, TEXT( "If enabled, the octree for editable meshes will have a huge bounding box.  Otherwise, we'll compute a tightly wrapped bounds.  However, the bounds will not be able to grow beyond it's original size." ) );
 }
 
-
 //
 // =========================================================
 //
 
+namespace MeshAttribute
+{
+	const FName PolygonGroup::MaterialAssetName( "MaterialAssetName" );
+}
+
+
 /** Traits class which specifies whether Lerp is defined for a given type */
 template <typename T> struct TIsAttributeTypeLerpable { static const bool Value = false; };
 template <> struct TIsAttributeTypeLerpable<FVector4> { static const bool Value = true; };
@@ -590,22 +595,6 @@ void UEditableMesh::RebuildRenderMesh()
 	}
 }
 
-bool UEditableMesh::FindOrAddMaterial(class UMaterialInterface* Material, int32& MaterialIndex, FName& ImportedMaterialName)
-{
-	//Only one adapter should add a material
-	//TODO: Do we want to support editing a SM and a SK sharing a MeshDescription?
-	for (UEditableMeshAdapter* Adapter : Adapters)
-	{
-		if (Adapter->FindOrAddMaterial(Material, MaterialIndex, ImportedMaterialName))
-		{
-			if (MaterialIndex != INDEX_NONE)
-			{
-				return true;
-			}
-		}
-	}
-	return false;
-}
 
 void UEditableMesh::StartModification( const EMeshModificationType MeshModificationType, const EMeshTopologyChange MeshTopologyChange )
 {
diff --git a/Engine/Plugins/Runtime/EditableMesh/Source/EditableMesh/EditableStaticMeshAdapter.cpp b/Engine/Plugins/Runtime/EditableMesh/Source/EditableMesh/EditableStaticMeshAdapter.cpp
--- a/Engine/Plugins/Runtime/EditableMesh/Source/EditableMesh/EditableStaticMeshAdapter.cpp
+++ b/Engine/Plugins/Runtime/EditableMesh/Source/EditableMesh/EditableStaticMeshAdapter.cpp
@@ -105,6 +105,7 @@ void UEditableStaticMeshAdapter::InitEditableStaticMesh( UEditableMesh* Editable
 				TEdgeAttributeArray<bool>& EdgeHardnesses = MeshDescription->EdgeAttributes().GetAttributes<bool>( MeshAttribute::Edge::IsHard );
 
 				TPolygonGroupAttributeArray<FName>& PolygonGroupImportedMaterialSlotNames = MeshDescription->PolygonGroupAttributes().GetAttributes<FName>( MeshAttribute::PolygonGroup::ImportedMaterialSlotName );
+				TPolygonGroupAttributeArray<FName>& PolygonGroupMaterialAssetNames = MeshDescription->PolygonGroupAttributes().GetAttributes<FName>( MeshAttribute::PolygonGroup::MaterialAssetName );
 				TPolygonGroupAttributeArray<bool>& PolygonGroupCollision = MeshDescription->PolygonGroupAttributes().GetAttributes<bool>( MeshAttribute::PolygonGroup::EnableCollision );
 				TPolygonGroupAttributeArray<bool>& PolygonGroupCastShadow = MeshDescription->PolygonGroupAttributes().GetAttributes<bool>( MeshAttribute::PolygonGroup::CastShadow );
 
@@ -275,6 +276,7 @@ void UEditableStaticMeshAdapter::InitEditableStaticMesh( UEditableMesh* Editable
 					// Create a new polygon group
 					const FPolygonGroupID NewPolygonGroupID = MeshDescription->CreatePolygonGroup();
 					PolygonGroupImportedMaterialSlotNames[ NewPolygonGroupID ] = StaticMaterial.ImportedMaterialSlotName;
+					PolygonGroupMaterialAssetNames[ NewPolygonGroupID ] = FName( *MaterialInterface->GetPathName() );
 					PolygonGroupCollision[ NewPolygonGroupID ] = RenderingSection.bEnableCollision;
 					PolygonGroupCastShadow[ NewPolygonGroupID ] = RenderingSection.bCastShadow;
 
@@ -287,6 +289,7 @@ void UEditableStaticMeshAdapter::InitEditableStaticMesh( UEditableMesh* Editable
 					NewRenderingPolygonGroup.Triangles.Reserve( NumSectionTriangles );
 					NewRenderingPolygonGroup.MaxTriangles = NumSectionTriangles;
 					NewRenderingPolygonGroup.RenderingSectionIndex = RenderingSectionIndex;
+					NewRenderingPolygonGroup.MaterialIndex = RenderingSection.MaterialIndex;
 
 					MeshDescription->ReserveNewPolygons( NumSectionTriangles );
 					MeshDescription->ReserveNewEdges( NumSectionTriangles * 3 );	// more than required, but not a problem
@@ -448,55 +451,6 @@ void UEditableStaticMeshAdapter::InitFromBlankStaticMesh( UEditableMesh* Editabl
 	StaticMesh = &InStaticMesh;
 }
 
-bool UEditableStaticMeshAdapter::FindOrAddMaterial(UMaterialInterface* Material, int32& MaterialIndex, FName& ImportedMaterialName)
-{
-	if (!StaticMesh)
-	{
-		return false;
-	}
-	//Look if the material already exist
-	for (int32 CurrentMaterialIndex = 0; CurrentMaterialIndex < StaticMesh->StaticMaterials.Num(); ++CurrentMaterialIndex)
-	{
-		if (StaticMesh->StaticMaterials[CurrentMaterialIndex].MaterialInterface == Material)
-		{
-			MaterialIndex = CurrentMaterialIndex;
-			ImportedMaterialName = StaticMesh->StaticMaterials[CurrentMaterialIndex].ImportedMaterialSlotName;
-			return true;
-		}
-	}
-	//If we did not found any existing material, create an entry in the static mesh material list and return the slot index
-	FName MaterialName = Material->GetFName();
-	bool bFoundMatch = true;
-	int32 Counter = 1;
-	while (bFoundMatch)
-	{
-		bFoundMatch = false;
-		for (const FStaticMaterial& StaticMaterial : StaticMesh->StaticMaterials)
-		{
-			if (StaticMaterial.ImportedMaterialSlotName == MaterialName)
-			{
-				bFoundMatch = true;
-				break;
-			}
-		}
-		if (bFoundMatch)
-		{
-			MaterialName = *FString::Printf(TEXT("%s_%d"), *Material->GetName(), Counter);
-			Counter++;
-		}
-	}
-
-	MaterialIndex = StaticMesh->StaticMaterials.Add(FStaticMaterial(
-		Material
-		, MaterialName
-#if WITH_EDITORONLY_DATA
-		, MaterialName
-#endif
-	));
-
-	ImportedMaterialName = MaterialName;
-	return true;
-}
 
 void UEditableStaticMeshAdapter::InitializeFromEditableMesh( const UEditableMesh* EditableMesh )
 {
@@ -522,6 +476,14 @@ void UEditableStaticMeshAdapter::InitializeFromEditableMesh( const UEditableMesh
 		RenderingPolygonGroups.Insert( PolygonGroupID );
 		FRenderingPolygonGroup& RenderingPolygonGroup = RenderingPolygonGroups[ PolygonGroupID ];
 		RenderingPolygonGroup.RenderingSectionIndex = RenderingSectionIndex;
+
+		const FName SlotName = MeshDescription->PolygonGroupAttributes().GetAttribute<FName>( PolygonGroupID, MeshAttribute::PolygonGroup::ImportedMaterialSlotName );
+		RenderingPolygonGroup.MaterialIndex = StaticMesh->StaticMaterials.IndexOfByPredicate(
+			[ &SlotName ]( const FStaticMaterial& StaticMaterial )
+			{
+				return StaticMaterial.ImportedMaterialSlotName == SlotName;
+			}
+		);
 		RenderingPolygonGroup.MaxTriangles = 0;
 
 		RenderingSectionIndex++;
@@ -840,8 +802,7 @@ void UEditableStaticMeshAdapter::OnStartModification( const UEditableMesh* Edita
 	// @todo mesheditor undo: We're not using traditional transactions to undo mesh changes yet, but we still want to dirty the mesh package
 	// Also, should we even need the Initializing type? Should we not wait for the first modification before dirtying the package?
 
-	StaticMesh->SetFlags( RF_Transactional );
-	StaticMesh->Modify();
+	StaticMesh->MarkPackageDirty();
 }
 
 
@@ -1701,29 +1662,24 @@ void UEditableStaticMeshAdapter::OnCreatePolygonGroups( const UEditableMesh* Edi
 	check( MeshDescription );
 
 	const TPolygonGroupAttributeArray<FName>& PolygonGroupImportedMaterialSlotNames = MeshDescription->PolygonGroupAttributes().GetAttributes<FName>( MeshAttribute::PolygonGroup::ImportedMaterialSlotName );
+	const TPolygonGroupAttributeArray<FName>& PolygonGroupMaterialAssetSlotNames = MeshDescription->PolygonGroupAttributes().GetAttributes<FName>( MeshAttribute::PolygonGroup::MaterialAssetName );
 	const TPolygonGroupAttributeArray<bool>& PolygonGroupCollision = MeshDescription->PolygonGroupAttributes().GetAttributes<bool>( MeshAttribute::PolygonGroup::EnableCollision );
 	const TPolygonGroupAttributeArray<bool>& PolygonGroupCastShadow = MeshDescription->PolygonGroupAttributes().GetAttributes<bool>( MeshAttribute::PolygonGroup::CastShadow );
 
 	for( const FPolygonGroupID PolygonGroupID : PolygonGroupIDs )
 	{
+		UMaterialInterface* Material = LoadObject<UMaterialInterface>( nullptr, *PolygonGroupMaterialAssetSlotNames[ PolygonGroupID ].ToString() );
+		const int32 MaterialIndex = StaticMesh->StaticMaterials.Emplace(
+			Material,
+			PolygonGroupImportedMaterialSlotNames[ PolygonGroupID ]
+#if WITH_EDITORONLY_DATA
+			, PolygonGroupImportedMaterialSlotNames[ PolygonGroupID ]
+#endif
+		);
+
 		uint32 LODSectionIndex = 0;
 		if( !EditableMesh->IsPreviewingSubdivisions() )
 		{
-			//Verify the meshdescription new polygon group data is correct, The attribute must match the static mesh material slot array
-			//The staticmesh material slot entry was create before executing this function
-			bool bFoundMaterial = false;
-			int32 PolygonGroupMaterialIndex = INDEX_NONE;
-			for (int32 MaterialIndex = 0; MaterialIndex < StaticMesh->StaticMaterials.Num(); ++MaterialIndex)
-			{
-				if (PolygonGroupImportedMaterialSlotNames[PolygonGroupID] == StaticMesh->StaticMaterials[MaterialIndex].ImportedMaterialSlotName)
-				{
-					PolygonGroupMaterialIndex = MaterialIndex;
-					bFoundMaterial = true;
-					break;
-				}
-			}
-			check(bFoundMaterial);
-
 			// Need to create a new rendering section. This is added to the end of the array.
 			FStaticMeshLODResources& StaticMeshLOD = GetStaticMeshLOD();
 
@@ -1755,7 +1711,7 @@ void UEditableStaticMeshAdapter::OnCreatePolygonGroups( const UEditableMesh* Edi
 			StaticMeshSection.MaxVertexIndex = 0;
 			StaticMeshSection.bEnableCollision = PolygonGroupCollision[ PolygonGroupID ];
 			StaticMeshSection.bCastShadow = PolygonGroupCastShadow[ PolygonGroupID ];
-			StaticMeshSection.MaterialIndex = PolygonGroupMaterialIndex;
+			StaticMeshSection.MaterialIndex = MaterialIndex;
 		}
 
 		// Insert the rendering polygon group for keeping track of these index buffer properties
@@ -1763,6 +1719,7 @@ void UEditableStaticMeshAdapter::OnCreatePolygonGroups( const UEditableMesh* Edi
 		FRenderingPolygonGroup& RenderingPolygonGroup = RenderingPolygonGroups[ PolygonGroupID ];
 
 		RenderingPolygonGroup.RenderingSectionIndex = LODSectionIndex;
+		RenderingPolygonGroup.MaterialIndex = MaterialIndex;
 		RenderingPolygonGroup.MaxTriangles = 0;
 	}
 }
@@ -1776,27 +1733,28 @@ void UEditableStaticMeshAdapter::OnSetPolygonGroupAttribute( const UEditableMesh
 	FStaticMeshLODResources& StaticMeshLOD = GetStaticMeshLOD();
 	FStaticMeshSection& StaticMeshSection = StaticMeshLOD.Sections[ RenderingPolygonGroup.RenderingSectionIndex ];
 
-	
-	const TPolygonGroupAttributeArray<FName>& PolygonGroupImportedMaterialSlotNames = MeshDescription->PolygonGroupAttributes().GetAttributes<FName>( MeshAttribute::PolygonGroup::ImportedMaterialSlotName );
-	
-	bool bFoundMaterial = false;
-	int32 PolygonGroupMaterialIndex = INDEX_NONE;
-	for (int32 MaterialIndex = 0; MaterialIndex < StaticMesh->StaticMaterials.Num(); ++MaterialIndex)
-	{
-		if (PolygonGroupImportedMaterialSlotNames[PolygonGroupID] == StaticMesh->StaticMaterials[MaterialIndex].ImportedMaterialSlotName)
-		{
-			PolygonGroupMaterialIndex = MaterialIndex;
-			bFoundMaterial = true;
-			break;
-		}
-	}
-	check(bFoundMaterial);
+	const FName ImportedMaterialSlotName = MeshDescription->PolygonGroupAttributes().GetAttribute<FName>( PolygonGroupID, MeshAttribute::PolygonGroup::ImportedMaterialSlotName, 0 );
+	const FName MaterialAssetName = MeshDescription->PolygonGroupAttributes().GetAttribute<FName>( PolygonGroupID, MeshAttribute::PolygonGroup::MaterialAssetName, 0 );
 
-	const int32 MaterialIndex = PolygonGroupMaterialIndex;
+	const int32 MaterialIndex = StaticMesh->StaticMaterials.IndexOfByPredicate(
+		[ &ImportedMaterialSlotName ]( const FStaticMaterial& StaticMaterial ) { return StaticMaterial.ImportedMaterialSlotName == ImportedMaterialSlotName; }
+	);
+	check( MaterialIndex != INDEX_NONE );
 
 	if ( Attribute.AttributeName == MeshAttribute::PolygonGroup::ImportedMaterialSlotName )
 	{
-		StaticMesh->StaticMaterials[MaterialIndex].ImportedMaterialSlotName = PolygonGroupImportedMaterialSlotNames[PolygonGroupID];
+		StaticMesh->StaticMaterials[ RenderingPolygonGroup.MaterialIndex ].ImportedMaterialSlotName = ImportedMaterialSlotName;
+	}
+	else if( Attribute.AttributeName == MeshAttribute::PolygonGroup::MaterialAssetName )
+	{
+		UMaterialInterface* Material = LoadObject<UMaterialInterface>( nullptr, *MaterialAssetName.ToString() );
+		StaticMesh->StaticMaterials[ RenderingPolygonGroup.MaterialIndex ] = FStaticMaterial(
+			Material,
+			ImportedMaterialSlotName
+#if WITH_EDITORONLY_DATA
+			, ImportedMaterialSlotName
+#endif
+		);
 	}
 	else if( Attribute.AttributeName == MeshAttribute::PolygonGroup::CastShadow )
 	{
@@ -1817,12 +1775,8 @@ void UEditableStaticMeshAdapter::OnDeletePolygonGroups( const UEditableMesh* Edi
 	{
 		FRenderingPolygonGroup& RenderingPolygonGroup = RenderingPolygonGroups[ PolygonGroupID ];
 
-		const FName SlotName = PolygonGroupImportedMaterialSlotNames[ PolygonGroupID ];
-
 		// Remove material slot associated with section
-		const int32 MaterialIndex = StaticMesh->StaticMaterials.IndexOfByPredicate(
-			[ &SlotName ]( const FStaticMaterial& StaticMaterial ) { return StaticMaterial.ImportedMaterialSlotName == SlotName; }
-		);
+		const int32 MaterialIndex = RenderingPolygonGroup.MaterialIndex;
 		StaticMesh->StaticMaterials.RemoveAt( MaterialIndex );
 
 		// Adjust rendering indices held by sections: any index above the one we just deleted now needs to be decremented.
@@ -1836,6 +1790,11 @@ void UEditableStaticMeshAdapter::OnDeletePolygonGroups( const UEditableMesh* Edi
 			{
 				PolygonGroupToAdjust.RenderingSectionIndex--;
 			}
+
+			if( PolygonGroupToAdjust.MaterialIndex > MaterialIndex )
+			{
+				PolygonGroupToAdjust.MaterialIndex--;
+			}
 		}
 
 		if( !EditableMesh->IsPreviewingSubdivisions() )
diff --git a/Engine/Plugins/Runtime/EditableMesh/Source/EditableMesh/EditableStaticMeshAdapter.h b/Engine/Plugins/Runtime/EditableMesh/Source/EditableMesh/EditableStaticMeshAdapter.h
--- a/Engine/Plugins/Runtime/EditableMesh/Source/EditableMesh/EditableStaticMeshAdapter.h
+++ b/Engine/Plugins/Runtime/EditableMesh/Source/EditableMesh/EditableStaticMeshAdapter.h
@@ -71,6 +71,10 @@ struct FRenderingPolygonGroup
 	UPROPERTY()
 	uint32 RenderingSectionIndex;
 
+	/** The material slot index assigned to this polygon group's material */
+	UPROPERTY()
+	int32 MaterialIndex;
+
 	/** Maximum number of triangles which have been reserved in the index buffer */
 	UPROPERTY()
 	int32 MaxTriangles;
@@ -96,6 +100,7 @@ struct FRenderingPolygonGroup
 	friend FArchive& operator<<( FArchive& Ar, FRenderingPolygonGroup& Section )
 	{
 		Ar << Section.RenderingSectionIndex;
+		Ar << Section.MaterialIndex;
 		Ar << Section.MaxTriangles;	// @todo mesheditor serialization: Should not need to be serialized if we triangulate after load
 		Ar << Section.Triangles;
 		return Ar;
@@ -121,8 +126,6 @@ public:
 	void InitEditableStaticMesh( UEditableMesh* EditableMesh, class UPrimitiveComponent& Component, const FEditableMeshSubMeshAddress& InitSubMeshAddress );
 	EDITABLEMESH_API void InitFromBlankStaticMesh( UEditableMesh* EditableMesh, UStaticMesh& InStaticMesh );
 
-	virtual bool FindOrAddMaterial(class UMaterialInterface* Material, int32& MaterialIndex, FName& ImportedMaterialName) override;
-
 	virtual void InitializeFromEditableMesh( const UEditableMesh* EditableMesh ) override;
 	virtual void OnRebuildRenderMeshStart( const UEditableMesh* EditableMesh, const bool bInvalidateLighting ) override;
 	virtual void OnRebuildRenderMesh( const UEditableMesh* EditableMesh ) override;
diff --git a/Engine/Plugins/Runtime/EditableMesh/Source/EditableMesh/Public/EditableMesh.h b/Engine/Plugins/Runtime/EditableMesh/Source/EditableMesh/Public/EditableMesh.h
--- a/Engine/Plugins/Runtime/EditableMesh/Source/EditableMesh/Public/EditableMesh.h
+++ b/Engine/Plugins/Runtime/EditableMesh/Source/EditableMesh/Public/EditableMesh.h
@@ -46,6 +46,18 @@ namespace OpenSubdiv
 #endif
 
 
+/**
+ * Additional mesh attributes required by EditableMesh
+ */
+namespace MeshAttribute
+{
+	namespace PolygonGroup
+	{
+		extern EDITABLEMESH_API const FName MaterialAssetName;
+	}
+}
+
+
 UENUM( BlueprintType )
 enum class EInsetPolygonsMode : uint8
 {
@@ -220,8 +232,6 @@ public:
 	UFUNCTION( BlueprintPure, Category="Editable Mesh" )
 	FMeshTriangle GetPolygonTriangulatedTriangle( const FPolygonID PolygonID, int32 PolygonTriangleNumber ) const;
 
-	bool FindOrAddMaterial(class UMaterialInterface* Material, int32& MaterialIndex, FName& ImportedMaterialName);
-
 protected:
 
 	/** Given a set of index remappings, fixes up references in the octree */
diff --git a/Engine/Plugins/Runtime/EditableMesh/Source/EditableMesh/Public/EditableMeshAdapter.h b/Engine/Plugins/Runtime/EditableMesh/Source/EditableMesh/Public/EditableMeshAdapter.h
--- a/Engine/Plugins/Runtime/EditableMesh/Source/EditableMesh/Public/EditableMeshAdapter.h
+++ b/Engine/Plugins/Runtime/EditableMesh/Source/EditableMesh/Public/EditableMeshAdapter.h
@@ -16,8 +16,6 @@ class EDITABLEMESH_API UEditableMeshAdapter : public UObject
 
 public:
 
-	virtual bool FindOrAddMaterial(class UMaterialInterface* Material, int32& MaterialIndex, FName& ImportedMaterialName) PURE_VIRTUAL(, return false;);
-
 	virtual void InitializeFromEditableMesh( const UEditableMesh* EditableMesh ) PURE_VIRTUAL(,);
 	virtual void OnRebuildRenderMeshStart( const UEditableMesh* EditableMesh, const bool bInvalidateLighting ) PURE_VIRTUAL(,);
 	virtual void OnRebuildRenderMesh( const UEditableMesh* EditableMesh ) PURE_VIRTUAL(,);
diff --git a/Engine/Plugins/Runtime/EditableMesh/Source/EditableMesh/StaticMeshEditableMeshFormat.cpp b/Engine/Plugins/Runtime/EditableMesh/Source/EditableMesh/StaticMeshEditableMeshFormat.cpp
--- a/Engine/Plugins/Runtime/EditableMesh/Source/EditableMesh/StaticMeshEditableMeshFormat.cpp
+++ b/Engine/Plugins/Runtime/EditableMesh/Source/EditableMesh/StaticMeshEditableMeshFormat.cpp
@@ -45,7 +45,12 @@ UEditableMesh* FStaticMeshEditableMeshFormat::MakeEditableMesh( UPrimitiveCompon
 	}
 
 	UEditableMesh* EditableMesh = NewObject<UEditableMesh>();
-	EditableMesh->MeshDescription = NewObject<UMeshDescription>( EditableMesh );
+	UMeshDescription* MeshDescription = NewObject<UMeshDescription>( EditableMesh );
+	EditableMesh->MeshDescription = MeshDescription;
+	UStaticMesh::RegisterMeshAttributes( MeshDescription );
+
+	// Register additional attributes required by EditableMesh
+	MeshDescription->PolygonGroupAttributes().RegisterAttribute<FName>( MeshAttribute::PolygonGroup::MaterialAssetName );
 
 	UEditableStaticMeshAdapter* EditableStaticMesh = NewObject<UEditableStaticMeshAdapter>( EditableMesh );
 	EditableMesh->Adapters.Add( EditableStaticMesh );
diff --git a/Engine/Plugins/Tests/EditorTests/Source/EditorTests/Private/UnrealEd/MeshDescriptionTests.cpp b/Engine/Plugins/Tests/EditorTests/Source/EditorTests/Private/UnrealEd/MeshDescriptionTests.cpp
--- a/Engine/Plugins/Tests/EditorTests/Source/EditorTests/Private/UnrealEd/MeshDescriptionTests.cpp
+++ b/Engine/Plugins/Tests/EditorTests/Source/EditorTests/Private/UnrealEd/MeshDescriptionTests.cpp
@@ -535,6 +535,7 @@ bool FMeshDescriptionTest::ConversionTest(FAutomationTestExecutionInfo& Executio
 				AssetMesh->SourceModels[LodIndex].LoadRawMesh(ResultRawMesh);
 				//Create a temporary Mesh Description
 				UMeshDescription* MeshDescription = NewObject<UMeshDescription>(GetTransientPackage(), NAME_None, RF_Standalone);
+				UStaticMesh::RegisterMeshAttributes(MeshDescription);
 				FMeshDescriptionOperations::ConverFromRawMesh(ResultRawMesh, MeshDescription, MaterialMapInverse);
 				//Convert back the FRawmesh
 				FMeshDescriptionOperations::ConverToRawMesh(MeshDescription, ResultRawMesh, MaterialMap);
diff --git a/Engine/Source/Developer/MeshBuilder/Private/StaticMeshBuilder.cpp b/Engine/Source/Developer/MeshBuilder/Private/StaticMeshBuilder.cpp
--- a/Engine/Source/Developer/MeshBuilder/Private/StaticMeshBuilder.cpp
+++ b/Engine/Source/Developer/MeshBuilder/Private/StaticMeshBuilder.cpp
@@ -78,6 +78,7 @@ bool FStaticMeshBuilder::Build(UStaticMesh* StaticMesh, const FStaticMeshLODGrou
 			{
 				//Create an empty mesh description that the reduce will fill
 				MeshDescription = NewObject<UMeshDescription>(StaticMesh, NAME_None);
+				UStaticMesh::RegisterMeshAttributes(MeshDescription);
 			}
 			else if(StaticMesh->GetMeshDescription(0) != nullptr)
 			{
diff --git a/Engine/Source/Editor/UnrealEd/Private/Fbx/FbxStaticMeshImport.cpp b/Engine/Source/Editor/UnrealEd/Private/Fbx/FbxStaticMeshImport.cpp
--- a/Engine/Source/Editor/UnrealEd/Private/Fbx/FbxStaticMeshImport.cpp
+++ b/Engine/Source/Editor/UnrealEd/Private/Fbx/FbxStaticMeshImport.cpp
@@ -1345,6 +1345,7 @@ UStaticMesh* UnFbx::FFbxImporter::ImportStaticMeshAsSingle(UObject* InParent, TA
 		//Create private asset in the same package as the StaticMesh, and make sure reference are set to avoid GC
 		MeshDescription = NewObject<UMeshDescription>(StaticMesh, NAME_None, RF_NoFlags);
 		check(MeshDescription != nullptr);
+		UStaticMesh::RegisterMeshAttributes(MeshDescription);
 		StaticMesh->SetOriginalMeshDescription(LODIndex, MeshDescription);
 	}
 	
diff --git a/Engine/Source/Editor/UnrealEd/Private/StaticMeshEdit.cpp b/Engine/Source/Editor/UnrealEd/Private/StaticMeshEdit.cpp
--- a/Engine/Source/Editor/UnrealEd/Private/StaticMeshEdit.cpp
+++ b/Engine/Source/Editor/UnrealEd/Private/StaticMeshEdit.cpp
@@ -929,6 +929,7 @@ UStaticMesh* CreateStaticMeshFromBrush(UObject* Outer,FName Name,ABrush* Brush,U
 	GWarn->BeginSlowTask( NSLOCTEXT("UnrealEd", "CreatingStaticMeshE", "Creating static mesh..."), true );
 	UStaticMesh* StaticMesh = nullptr;
 	UMeshDescription* MeshDescription = NewObject<UMeshDescription>(Outer);
+	UStaticMesh::RegisterMeshAttributes(MeshDescription);
 	TArray<FStaticMaterial> Materials;
 	GetBrushMesh(Brush, Model, MeshDescription, Materials);
 	StaticMesh = CreateStaticMesh(MeshDescription, Materials, Outer, Name);
diff --git a/Engine/Source/Runtime/Engine/Classes/Engine/StaticMesh.h b/Engine/Source/Runtime/Engine/Classes/Engine/StaticMesh.h
--- a/Engine/Source/Runtime/Engine/Classes/Engine/StaticMesh.h
+++ b/Engine/Source/Runtime/Engine/Classes/Engine/StaticMesh.h
@@ -679,6 +679,12 @@ public:
 	 */
 	UPROPERTY(VisibleAnywhere, Category = MeshDescription)
 	TArray<class UMeshDescription*> MeshDescriptions;
+
+	/**
+	 * Registers the mesh attributes required by the mesh description for a static mesh.
+	 */
+	ENGINE_API static void RegisterMeshAttributes( UMeshDescription* MeshDescription );
+
 	/**
 	 * Accessors for the mesh description use to build the render data
 	 */
diff --git a/Engine/Source/Runtime/Engine/Private/StaticMesh.cpp b/Engine/Source/Runtime/Engine/Private/StaticMesh.cpp
--- a/Engine/Source/Runtime/Engine/Private/StaticMesh.cpp
+++ b/Engine/Source/Runtime/Engine/Private/StaticMesh.cpp
@@ -2521,6 +2521,38 @@ static FStaticMeshRenderData& GetPlatformStaticMeshRenderData(UStaticMesh* Mesh,
 
 #if WITH_EDITORONLY_DATA
 
+/**
+ * Registers the mesh attributes required by the mesh description for a static mesh.
+ */
+void UStaticMesh::RegisterMeshAttributes( UMeshDescription* MeshDescription )
+{
+	// Add basic vertex attributes
+	MeshDescription->VertexAttributes().RegisterAttribute<FVector>( MeshAttribute::Vertex::Position, 1, FVector::ZeroVector, EMeshAttributeFlags::Lerpable );
+	MeshDescription->VertexAttributes().RegisterAttribute<float>( MeshAttribute::Vertex::CornerSharpness, 1, 0.0f, EMeshAttributeFlags::Lerpable );
+
+	// Add basic vertex instance attributes
+	MeshDescription->VertexInstanceAttributes().RegisterAttribute<FVector2D>( MeshAttribute::VertexInstance::TextureCoordinate, 1, FVector2D::ZeroVector, EMeshAttributeFlags::Lerpable );
+	MeshDescription->VertexInstanceAttributes().RegisterAttribute<FVector>( MeshAttribute::VertexInstance::Normal, 1, FVector::ZeroVector, EMeshAttributeFlags::AutoGenerated );
+	MeshDescription->VertexInstanceAttributes().RegisterAttribute<FVector>( MeshAttribute::VertexInstance::Tangent, 1, FVector::ZeroVector, EMeshAttributeFlags::AutoGenerated );
+	MeshDescription->VertexInstanceAttributes().RegisterAttribute<float>( MeshAttribute::VertexInstance::BinormalSign, 1, 0.0f, EMeshAttributeFlags::AutoGenerated );
+	MeshDescription->VertexInstanceAttributes().RegisterAttribute<FVector4>( MeshAttribute::VertexInstance::Color, 1, FVector4( 1.0f ), EMeshAttributeFlags::Lerpable );
+
+	// Add basic edge attributes
+	MeshDescription->EdgeAttributes().RegisterAttribute<bool>( MeshAttribute::Edge::IsHard, 1, false );
+	MeshDescription->EdgeAttributes().RegisterAttribute<float>( MeshAttribute::Edge::CreaseSharpness, 1, 0.0f, EMeshAttributeFlags::Lerpable );
+
+	// Add basic polygon attributes
+	MeshDescription->PolygonAttributes().RegisterAttribute<FVector>( MeshAttribute::Polygon::Normal, 1, FVector::ZeroVector, EMeshAttributeFlags::AutoGenerated );
+	MeshDescription->PolygonAttributes().RegisterAttribute<FVector>( MeshAttribute::Polygon::Tangent, 1, FVector::ZeroVector, EMeshAttributeFlags::AutoGenerated );
+	MeshDescription->PolygonAttributes().RegisterAttribute<FVector>( MeshAttribute::Polygon::Binormal, 1, FVector::ZeroVector, EMeshAttributeFlags::AutoGenerated );
+	MeshDescription->PolygonAttributes().RegisterAttribute<FVector>( MeshAttribute::Polygon::Center, 1, FVector::ZeroVector, EMeshAttributeFlags::AutoGenerated );
+
+	// Add basic polygon group attributes
+	MeshDescription->PolygonGroupAttributes().RegisterAttribute<FName>( MeshAttribute::PolygonGroup::ImportedMaterialSlotName ); //The unique key to match the mesh material slot
+	MeshDescription->PolygonGroupAttributes().RegisterAttribute<bool>( MeshAttribute::PolygonGroup::EnableCollision ); //Deprecated
+	MeshDescription->PolygonGroupAttributes().RegisterAttribute<bool>( MeshAttribute::PolygonGroup::CastShadow ); //Deprecated
+}
+
 UMeshDescription* UStaticMesh::GetMeshDescription(int32 LodIndex) const
 {
 	if (!MeshDescriptions.IsValidIndex(LodIndex))
@@ -2554,6 +2586,8 @@ UMeshDescription* UStaticMesh::GetOriginalMeshDescription(int32 LodIndex)
 		{
 			//Fill the MeshDescription on demand
 			SourceModels[LodIndex].OriginalMeshDescription = NewObject<UMeshDescription>(this, NAME_None);
+			RegisterMeshAttributes(SourceModels[LodIndex].OriginalMeshDescription);
+
 			FRawMesh LodRawMesh;
 			SourceModels[LodIndex].LoadRawMesh(LodRawMesh);
 			TMap<int32, FName> MaterialMap;
@@ -2645,7 +2679,7 @@ void UStaticMesh::FixupMaterialSlotName()
 // differences, etc.) replace the version GUID below with a new one.
 // In case of merge conflicts with DDC versions, you MUST generate a new GUID
 // and set this new GUID as the version.                                       
-#define MESHDATAKEY_STATICMESH_DERIVEDDATA_VER TEXT("65133A83EADF412DB9C57EF087814E02")
+#define MESHDATAKEY_STATICMESH_DERIVEDDATA_VER TEXT("8ABEA51A53B84861846CAE42E9C11538")
 
 static const FString& GetMeshDataKeyStaticMeshDerivedDataVersion()
 {
@@ -2831,8 +2865,11 @@ void UStaticMesh::CacheMeshData()
 				FRawMesh TempRawMesh;
 				//Convert the rawmesh to mesh description mark the ddc for save
 				SourceModel.RawMeshBulkData->LoadRawMesh(TempRawMesh);
+
 				//The original mesh description is null we must create one
 				SourceModel.OriginalMeshDescription = NewObject<UMeshDescription>(this, NAME_None);
+				RegisterMeshAttributes(SourceModel.OriginalMeshDescription);
+
 				//Convert the raw mesh to mesh description and save the mesh description in the ddc
 				TMap<int32, FName> MaterialMap;
 				for (int32 MaterialIndex = 0; MaterialIndex < StaticMaterials.Num(); ++MaterialIndex)
diff --git a/Engine/Source/Runtime/MeshDescription/Private/MeshDescription.cpp b/Engine/Source/Runtime/MeshDescription/Private/MeshDescription.cpp
--- a/Engine/Source/Runtime/MeshDescription/Private/MeshDescription.cpp
+++ b/Engine/Source/Runtime/MeshDescription/Private/MeshDescription.cpp
@@ -9,66 +9,8 @@
 
 UMeshDescription::UMeshDescription()
 {
-	RegisterBaseAttributes();
 }
 
-void UMeshDescription::RegisterBaseAttributes()
-{
-	// Add basic vertex attributes
-	VertexAttributes().RegisterAttribute<FVector>(MeshAttribute::Vertex::Position, 1, FVector::ZeroVector, EMeshAttributeFlags::Lerpable);
-	VertexAttributes().RegisterAttribute<float>(MeshAttribute::Vertex::CornerSharpness, 1, 0.0f, EMeshAttributeFlags::Lerpable);
-
-	// Add basic vertex instance attributes
-	VertexInstanceAttributes().RegisterAttribute<FVector2D>(MeshAttribute::VertexInstance::TextureCoordinate, 1, FVector2D::ZeroVector, EMeshAttributeFlags::Lerpable);
-	VertexInstanceAttributes().RegisterAttribute<FVector>(MeshAttribute::VertexInstance::Normal, 1, FVector::ZeroVector, EMeshAttributeFlags::AutoGenerated);
-	VertexInstanceAttributes().RegisterAttribute<FVector>(MeshAttribute::VertexInstance::Tangent, 1, FVector::ZeroVector, EMeshAttributeFlags::AutoGenerated);
-	VertexInstanceAttributes().RegisterAttribute<float>(MeshAttribute::VertexInstance::BinormalSign, 1, 0.0f, EMeshAttributeFlags::AutoGenerated);
-	VertexInstanceAttributes().RegisterAttribute<FVector4>(MeshAttribute::VertexInstance::Color, 1, FVector4(1.0f), EMeshAttributeFlags::Lerpable);
-
-	// Add basic edge attributes
-	EdgeAttributes().RegisterAttribute<bool>(MeshAttribute::Edge::IsHard, 1, false);
-	EdgeAttributes().RegisterAttribute<float>(MeshAttribute::Edge::CreaseSharpness, 1, 0.0f, EMeshAttributeFlags::Lerpable);
-
-	// Add basic polygon attributes
-	PolygonAttributes().RegisterAttribute<FVector>(MeshAttribute::Polygon::Normal, 1, FVector::ZeroVector, EMeshAttributeFlags::AutoGenerated);
-	PolygonAttributes().RegisterAttribute<FVector>(MeshAttribute::Polygon::Tangent, 1, FVector::ZeroVector, EMeshAttributeFlags::AutoGenerated);
-	PolygonAttributes().RegisterAttribute<FVector>(MeshAttribute::Polygon::Binormal, 1, FVector::ZeroVector, EMeshAttributeFlags::AutoGenerated);
-	PolygonAttributes().RegisterAttribute<FVector>(MeshAttribute::Polygon::Center, 1, FVector::ZeroVector, EMeshAttributeFlags::AutoGenerated);
-
-	// Add basic polygon group attributes
-	PolygonGroupAttributes().RegisterAttribute<FName>(MeshAttribute::PolygonGroup::ImportedMaterialSlotName); //The unique key to match the mesh material slot
-	PolygonGroupAttributes().RegisterAttribute<bool>(MeshAttribute::PolygonGroup::EnableCollision); //Deprecated
-	PolygonGroupAttributes().RegisterAttribute<bool>(MeshAttribute::PolygonGroup::CastShadow); //Deprecated
-}
-
-void UMeshDescription::UnRegisterBaseAttributes()
-{
-	// Remove basic vertex attributes
-	VertexAttributes().UnregisterAttribute<FVector>(MeshAttribute::Vertex::Position);
-	VertexAttributes().UnregisterAttribute<float>(MeshAttribute::Vertex::CornerSharpness);
-
-	// Remove basic vertex instance attributes
-	VertexInstanceAttributes().UnregisterAttribute<FVector2D>(MeshAttribute::VertexInstance::TextureCoordinate);
-	VertexInstanceAttributes().UnregisterAttribute<FVector>(MeshAttribute::VertexInstance::Normal);
-	VertexInstanceAttributes().UnregisterAttribute<FVector>(MeshAttribute::VertexInstance::Tangent);
-	VertexInstanceAttributes().UnregisterAttribute<float>(MeshAttribute::VertexInstance::BinormalSign);
-	VertexInstanceAttributes().UnregisterAttribute<FVector4>(MeshAttribute::VertexInstance::Color);
-
-	// Remove basic edge attributes
-	EdgeAttributes().UnregisterAttribute<bool>(MeshAttribute::Edge::IsHard);
-	EdgeAttributes().UnregisterAttribute<float>(MeshAttribute::Edge::CreaseSharpness);
-
-	// Remove basic polygon attributes
-	PolygonAttributes().UnregisterAttribute<FVector>(MeshAttribute::Polygon::Normal);
-	PolygonAttributes().UnregisterAttribute<FVector>(MeshAttribute::Polygon::Tangent);
-	PolygonAttributes().UnregisterAttribute<FVector>(MeshAttribute::Polygon::Binormal);
-	PolygonAttributes().UnregisterAttribute<FVector>(MeshAttribute::Polygon::Center);
-
-	// Remove basic polygon group attributes
-	PolygonGroupAttributes().UnregisterAttribute<FName>(MeshAttribute::PolygonGroup::ImportedMaterialSlotName);
-	PolygonGroupAttributes().UnregisterAttribute<bool>(MeshAttribute::PolygonGroup::EnableCollision);
-	PolygonGroupAttributes().UnregisterAttribute<bool>(MeshAttribute::PolygonGroup::CastShadow);
-}
 
 void UMeshDescription::Serialize( FArchive& Ar )
 {
diff --git a/Engine/Source/Runtime/MeshDescription/Public/MeshAttributeArray.h b/Engine/Source/Runtime/MeshDescription/Public/MeshAttributeArray.h
--- a/Engine/Source/Runtime/MeshDescription/Public/MeshAttributeArray.h
+++ b/Engine/Source/Runtime/MeshDescription/Public/MeshAttributeArray.h
@@ -500,78 +500,12 @@ public:
 private:
 	friend FArchive& operator<<( FArchive& Ar, TAttributesMap& AttributesMap )
 	{
-		struct FRegisteredAttribute
-		{
-			FName Name;
-			int32 NumIndices;
-			AttributeType DefaultValue;
-			EMeshAttributeFlags Flags;
-		};
-
-		// If we are loading, get a list of currently registered attributes and their metadata.
-		// This will be used to update the loaded attribute map after serialization.
-		TArray<FRegisteredAttribute> RegisteredAttributes;
-		if( Ar.IsLoading() )
-		{
-			for( const auto& AttributeNameAndIndicesArray : AttributesMap.Map )
-			{
-				RegisteredAttributes.Emplace();
-				FRegisteredAttribute& RegisteredAttribute = RegisteredAttributes.Last();
-				RegisteredAttribute.Name = AttributeNameAndIndicesArray.Key;
-				RegisteredAttribute.NumIndices = AttributeNameAndIndicesArray.Value.GetNumIndices();
-				RegisteredAttribute.DefaultValue = AttributeNameAndIndicesArray.Value.GetDefaultValue();
-				RegisteredAttribute.Flags = AttributeNameAndIndicesArray.Value.GetFlags();
-			}
-		}
-
 		// First serialize the number of elements which each attribute should contain
 		Ar << AttributesMap.NumElements;
 
 		// Now serialize the attributes of this type.
 		Ar << AttributesMap.Map;
 
-		// If we are loading, the attributes just serialized may be out of sync with the currently registered attributes, so we need to fix this up now.
-		if( Ar.IsLoading() )
-		{
-			// Go through the list of registered attributes we just built
-			for( const FRegisteredAttribute& RegisteredAttribute : RegisteredAttributes )
-			{
-				AttributeIndicesArrayType* AttributeIndicesArray = AttributesMap.Map.Find( RegisteredAttribute.Name );
-				if( !AttributeIndicesArray )
-				{
-					// Re-register any attributes that aren't there
-					AttributeIndicesArrayType NewAttributeIndicesArray( RegisteredAttribute.NumIndices, RegisteredAttribute.DefaultValue, RegisteredAttribute.Flags, AttributesMap.NumElements );
-
-					UE_LOG( LogMeshDescription, Log, TEXT( "Didn't find attribute '%s' - adding" ), *RegisteredAttribute.Name.ToString() );
-					AttributesMap.Map.Add( RegisteredAttribute.Name, MoveTemp( NewAttributeIndicesArray ) );
-				}
-				else if( AttributeIndicesArray->GetNumIndices() != RegisteredAttribute.NumIndices )
-				{
-					// Amend the number of attribute indices for any which are not correct
-					UE_LOG( LogMeshDescription, Log, TEXT( "Found attribute '%s' with the wrong number of indices - amending" ), *RegisteredAttribute.Name.ToString() );
-					AttributeIndicesArray->SetNumIndices( RegisteredAttribute.NumIndices );
-				}
-				else
-				{
-//					UE_LOG( LogMeshDescription, Log, TEXT( "Found attribute '%s' OK" ), *RegisteredAttribute.Name.ToString() );
-					check( AttributeIndicesArray->GetArrayForIndex( 0 ).Num() == AttributesMap.NumElements );
-				}
-			}
-
-			// Finally, go through all the attributes in the container, and remove any entries which are not in the registered list
-			for( auto AttributesIt = AttributesMap.Map.CreateIterator(); AttributesIt; ++AttributesIt )
-			{
-				const FName& AttributeName = AttributesIt.Key();
-				if( !RegisteredAttributes.ContainsByPredicate(
-					[ &AttributeName ]( const FRegisteredAttribute& RegisteredAttribute ) { return RegisteredAttribute.Name == AttributeName; }
-				) )
-				{
-					UE_LOG( LogMeshDescription, Warning, TEXT( "Found unregistered attribute '%s' - removing" ), *AttributeName.ToString() );
-					AttributesIt.RemoveCurrent();
-				}
-			}
-		}
-
 		return Ar;
 	}
 
diff --git a/Engine/Source/Runtime/MeshDescription/Public/MeshDescription.h b/Engine/Source/Runtime/MeshDescription/Public/MeshDescription.h
--- a/Engine/Source/Runtime/MeshDescription/Public/MeshDescription.h
+++ b/Engine/Source/Runtime/MeshDescription/Public/MeshDescription.h
@@ -230,7 +230,9 @@ using FEdgeArray = TMeshElementArray<FMeshEdge, FEdgeID>;
 using FPolygonArray = TMeshElementArray<FMeshPolygon, FPolygonID>;
 using FPolygonGroupArray = TMeshElementArray<FMeshPolygonGroup, FPolygonGroupID>;
 
-#define MESHDESCRIPTION_VER TEXT("C173EAB944C94C608318FAB504BD0F81")
+
+#define MESHDESCRIPTION_VER TEXT("264A91765090413DA61FE2A819147655")
+
 
 UCLASS()
 class MESHDESCRIPTION_API UMeshDescription : public UObject
@@ -814,8 +816,6 @@ private:
 	/** Given a set of index remappings, remaps all attributes accordingly */
 	void RemapAttributes( const FElementIDRemappings& Remappings );
 
-	void RegisterBaseAttributes();
-	void UnRegisterBaseAttributes();
 
 	FVertexArray VertexArray;
 	FVertexInstanceArray VertexInstanceArray;
