Added SetMeshDescription() to EditableMesh, so that it can use an arbitrary mesh description, not just its owned one.

diff --git a/Engine/Plugins/Runtime/EditableMesh/Source/EditableMesh/EditableMesh.cpp b/Engine/Plugins/Runtime/EditableMesh/Source/EditableMesh/EditableMesh.cpp
--- a/Engine/Plugins/Runtime/EditableMesh/Source/EditableMesh/EditableMesh.cpp
+++ b/Engine/Plugins/Runtime/EditableMesh/Source/EditableMesh/EditableMesh.cpp
@@ -359,7 +359,8 @@ static void BackupAttributesInList( FMeshElementAttributeList& BackupList, const
 
 
 UEditableMesh::UEditableMesh()
-	: bAllowUndo( false ),
+	: MeshDescription( &OwnedMeshDescription ),
+	  bAllowUndo( false ),
 	  bAllowCompact( false ),
 	  PendingCompactCounter( 0 ),
 	  bAllowSpatialDatabase( false )
@@ -372,7 +373,7 @@ void UEditableMesh::Serialize( FArchive& Ar )
 	Super::Serialize( Ar );
 	Ar.UsingCustomVersion( FEditableMeshCustomVersion::GUID );
 
-	Ar << MeshDescription;
+	Ar << OwnedMeshDescription;
 
 	// If the array was serialized containing any editor-only transient adapters, they will appear here as null, so remove them.
 	if( Ar.IsLoading() )
@@ -394,6 +395,13 @@ void UEditableMesh::PostLoad()
 }
 
 
+void UEditableMesh::SetMeshDescription( FMeshDescription* InMeshDescription )
+{
+	MeshDescription = InMeshDescription;
+	OwnedMeshDescription.Empty();
+}
+
+
 template <typename ElementIDType>
 static void InvertRemapTable( TSparseArray<ElementIDType>& InvertedRemapTable, const TSparseArray<ElementIDType>& RemapTable )
 {
diff --git a/Engine/Plugins/Runtime/EditableMesh/Source/EditableMesh/Public/EditableMesh.h b/Engine/Plugins/Runtime/EditableMesh/Source/EditableMesh/Public/EditableMesh.h
--- a/Engine/Plugins/Runtime/EditableMesh/Source/EditableMesh/Public/EditableMesh.h
+++ b/Engine/Plugins/Runtime/EditableMesh/Source/EditableMesh/Public/EditableMesh.h
@@ -98,8 +98,8 @@ public:
 	/** Remaps mesh element arrays according to the provided remappings, in order to undo a compact operation */
 	void Uncompact( const FElementIDRemappings& Remappings );
 
-	const FMeshDescription* GetMeshDescription() const { return &MeshDescription; }
-	FMeshDescription* GetMeshDescription() { return &MeshDescription; }
+	FMeshDescription* GetMeshDescription() const { return MeshDescription; }
+	void SetMeshDescription( FMeshDescription* InMeshDescription );
 
 	UFUNCTION( BlueprintCallable, Category="Editable Mesh" ) void InitializeAdapters();
 	UFUNCTION( BlueprintCallable, Category="Editable Mesh" ) void RebuildRenderMesh();
@@ -535,7 +535,11 @@ public:
 
 public:
 
-	FMeshDescription MeshDescription;
+	/** Pointer to the active mesh description for this editable mesh */
+	FMeshDescription* MeshDescription;
+
+	/** Owned mesh description for this editable mesh */
+	FMeshDescription OwnedMeshDescription;
 
 // @todo mesheditor: sort out member access. Currently StaticMesh adapter relies on accessing this stuff directly
 //protected:
