Prevented the StaticMeshEditor Edit Mode from always regenerating collision meshes. Now, they are regenerated only on Delete/Create polygons operations.

diff --git a/Engine/Plugins/Runtime/EditableMesh/Source/EditableMesh/EditableStaticMeshAdapter.cpp b/Engine/Plugins/Runtime/EditableMesh/Source/EditableMesh/EditableStaticMeshAdapter.cpp
--- a/Engine/Plugins/Runtime/EditableMesh/Source/EditableMesh/EditableStaticMeshAdapter.cpp
+++ b/Engine/Plugins/Runtime/EditableMesh/Source/EditableMesh/EditableStaticMeshAdapter.cpp
@@ -19,7 +19,8 @@ UEditableStaticMeshAdapter::UEditableStaticMeshAdapter()
 	: StaticMesh( nullptr ),
 	  StaticMeshLODIndex( 0 ),
 	  RecreateRenderStateContext(),
-	  CachedBoundingBoxAndSphere( FVector::ZeroVector, FVector::ZeroVector, 0.0f )
+	  CachedBoundingBoxAndSphere( FVector::ZeroVector, FVector::ZeroVector, 0.0f ),
+	  bUpdateCollisionNeeded( false )
 {
 }
 
@@ -793,7 +794,7 @@ void UEditableStaticMeshAdapter::OnRebuildRenderMeshStart( const UEditableMesh*
 
 void UEditableStaticMeshAdapter::OnEndModification( const UEditableMesh* EditableMesh )
 {
-	// nothing to do here
+	bUpdateCollisionNeeded = false;
 }
 
 
@@ -806,7 +807,8 @@ void UEditableStaticMeshAdapter::OnRebuildRenderMeshFinish( const UEditableMesh*
 
 	UpdateBounds( EditableMesh, bRebuildBoundsAndCollision );
 	
-	if( bRebuildBoundsAndCollision )
+	// Only UpdateCollision if there were Delete operations (or Create from undoing a Delete)
+	if( bUpdateCollisionNeeded && bRebuildBoundsAndCollision )
 	{
 		UpdateCollision();
 	}
@@ -1297,6 +1299,8 @@ void UEditableStaticMeshAdapter::OnCreatePolygons( const UEditableMesh* Editable
 	{
 		RenderingPolygons.Insert( PolygonID );
 		RenderingPolygons[ PolygonID ].PolygonGroupID = EditableMesh->GetGroupForPolygon( PolygonID );
+
+		bUpdateCollisionNeeded = true;
 	}
 }
 
@@ -1497,6 +1501,9 @@ void UEditableStaticMeshAdapter::OnDeletePolygons( const UEditableMesh* Editable
 
 		// Delete the polygon from the static mesh adapter mirror
 		RenderingPolygons.Remove( PolygonID );
+
+		// Flag only UEditableMesh::DeletePolygons operations, not DeletePolygonTriangles because it's also called for other operations
+		bUpdateCollisionNeeded = true;
 	}
 }
 
diff --git a/Engine/Plugins/Runtime/EditableMesh/Source/EditableMesh/EditableStaticMeshAdapter.h b/Engine/Plugins/Runtime/EditableMesh/Source/EditableMesh/EditableStaticMeshAdapter.h
--- a/Engine/Plugins/Runtime/EditableMesh/Source/EditableMesh/EditableStaticMeshAdapter.h
+++ b/Engine/Plugins/Runtime/EditableMesh/Source/EditableMesh/EditableStaticMeshAdapter.h
@@ -220,4 +220,7 @@ private:
 	/** Cached bounding box for the static mesh.  This bounds can be (temporarily) larger than the actual mesh itself as
 	    an optimization. */
 	FBoxSphereBounds CachedBoundingBoxAndSphere;
+
+	/** Flag to indicate if a modification requires regenerating collision */
+	bool bUpdateCollisionNeeded;
 };
