Fixed crash when indirectly removing a polygon group
  - Cause: The deletion of the material slot associated to the deleted polygon group was not taking in account the other LODs held by the static mesh.
  - Solution: Removed code deleting material slot associated to deleted polygon group.
  - Note: The user does have the ability to delte such material in the static mesh editor when no part of the static mesh is using this material slot.

diff --git a/Engine/Plugins/Runtime/EditableMesh/Source/EditableMesh/EditableStaticMeshAdapter.cpp b/Engine/Plugins/Runtime/EditableMesh/Source/EditableMesh/EditableStaticMeshAdapter.cpp
--- a/Engine/Plugins/Runtime/EditableMesh/Source/EditableMesh/EditableStaticMeshAdapter.cpp
+++ b/Engine/Plugins/Runtime/EditableMesh/Source/EditableMesh/EditableStaticMeshAdapter.cpp
@@ -1757,11 +1757,7 @@ void UEditableStaticMeshAdapter::OnDeletePolygonGroups( const UEditableMesh* Edi
 	{
 		FRenderingPolygonGroup& RenderingPolygonGroup = RenderingPolygonGroups[ PolygonGroupID ];
 
-		// Remove material slot associated with section
-		const int32 MaterialIndex = RenderingPolygonGroup.MaterialIndex;
-		StaticMesh->StaticMaterials.RemoveAt( MaterialIndex );
-
-		// Adjust rendering indices held by sections: any index above the one we just deleted now needs to be decremented.
+		// Adjust rendering indices held by sections: any index above the one we are about to delete now needs to be decremented.
 		const uint32 RenderingSectionIndex = RenderingPolygonGroup.RenderingSectionIndex;
 
 		for( const FPolygonGroupID PolygonGroupIDToAdjust : RenderingPolygonGroups.GetElementIDs() )
@@ -1772,11 +1768,6 @@ void UEditableStaticMeshAdapter::OnDeletePolygonGroups( const UEditableMesh* Edi
 			{
 				PolygonGroupToAdjust.RenderingSectionIndex--;
 			}
-
-			if( PolygonGroupToAdjust.MaterialIndex > MaterialIndex )
-			{
-				PolygonGroupToAdjust.MaterialIndex--;
-			}
 		}
 
 		if( !EditableMesh->IsPreviewingSubdivisions() )
@@ -1787,7 +1778,7 @@ void UEditableStaticMeshAdapter::OnDeletePolygonGroups( const UEditableMesh* Edi
 			// Get current number of triangles allocated for this section
 			const int32 MaxTriangles = RenderingPolygonGroup.MaxTriangles;
 
-			// Remove indices from this poisition in the index buffer
+			// Remove indices from this position in the index buffer
 			StaticMeshLOD.IndexBuffer.RemoveIndicesAt( FirstIndex, MaxTriangles * 3 );
 
 			// Adjust first index for all subsequent render sections to account for the indices just removed.
@@ -1799,23 +1790,6 @@ void UEditableStaticMeshAdapter::OnDeletePolygonGroups( const UEditableMesh* Edi
 				StaticMeshLOD.Sections[ Index ].FirstIndex -= MaxTriangles * 3;
 			}
 
-			// Adjust material indices for any sections to account for the fact that one has been removed
-			for( uint32 Index = 0; Index < NumRenderingSections; ++Index )
-			{
-				FStaticMeshSection& StaticMeshSection = StaticMeshLOD.Sections[ Index ];
-				if( StaticMeshSection.MaterialIndex > MaterialIndex )
-				{
-					StaticMeshSection.MaterialIndex--;
-
-#if WITH_EDITORONLY_DATA
-					// SectionInfoMap must be synced with the info of the modified Section
-					FMeshSectionInfo SectionInfo = StaticMesh->SectionInfoMap.Get( StaticMeshLODIndex, Index );
-					--SectionInfo.MaterialIndex;
-					StaticMesh->SectionInfoMap.Set( StaticMeshLODIndex, Index, SectionInfo );
-#endif
-				}
-			}
-
 			StaticMeshLOD.Sections.RemoveAt( RenderingSectionIndex );
 
 #if WITH_EDITORONLY_DATA
