Fixed crash on re-import of static mesh
  - Listening now to reimport at the toolbar level. Exiting edit mode if re-import happens
Improved a little bit messaging when non-editable LOD is selected
Fixed issue with display update when non-zero LOD is edited

diff --git a/Engine/Plugins/Runtime/EditableMesh/Source/EditableMesh/EditableStaticMeshAdapter.cpp b/Engine/Plugins/Runtime/EditableMesh/Source/EditableMesh/EditableStaticMeshAdapter.cpp
--- a/Engine/Plugins/Runtime/EditableMesh/Source/EditableMesh/EditableStaticMeshAdapter.cpp
+++ b/Engine/Plugins/Runtime/EditableMesh/Source/EditableMesh/EditableStaticMeshAdapter.cpp
@@ -411,8 +411,7 @@ void UEditableStaticMeshAdapter::InitializeFromEditableMesh( const UEditableMesh
 	// @todo mesheditor instancing: sort this out
 	OriginalStaticMesh = nullptr;
 
-	// Always targets LOD0 at the moment
-	StaticMeshLODIndex = 0;
+	StaticMeshLODIndex = SubMeshAddress.LODIndex;
 
 	RenderingPolygons.Reset();
 	RenderingPolygonGroups.Reset();
