Copied 4348861 Anousack.Kitisa
[StaticMeshEditor]:
* Fixed vertices being merged even though they don't share mergeable attributes.

diff --git a/Engine/Plugins/Runtime/EditableMesh/Source/EditableMesh/EditableMesh.cpp b/Engine/Plugins/Runtime/EditableMesh/Source/EditableMesh/EditableMesh.cpp
--- a/Engine/Plugins/Runtime/EditableMesh/Source/EditableMesh/EditableMesh.cpp
+++ b/Engine/Plugins/Runtime/EditableMesh/Source/EditableMesh/EditableMesh.cpp
@@ -257,6 +257,10 @@ static bool CompareAttributesByPredicate( const TAttributesSet<ElementIDType>& A
 					}
 				}
 			}
+			else
+			{
+				bResult = false;
+			}
 		}
 	);
 
