Copied 4383601 Anousack.Kitisa
[StaticMeshEditor]:
* Fixed material assignment issue when re-importing Datasmith static mesh after deleting a section from the mesh.
* Fixed material issue when deleting a section from a mesh and undoing and chaging material on the mesh.

diff --git a/Engine/Plugins/Enterprise/DatasmithContent/Source/DatasmithContent/Private/ObjectTemplates/DatasmithStaticMeshTemplate.cpp b/Engine/Plugins/Enterprise/DatasmithContent/Source/DatasmithContent/Private/ObjectTemplates/DatasmithStaticMeshTemplate.cpp
--- a/Engine/Plugins/Enterprise/DatasmithContent/Source/DatasmithContent/Private/ObjectTemplates/DatasmithStaticMeshTemplate.cpp
+++ b/Engine/Plugins/Enterprise/DatasmithContent/Source/DatasmithContent/Private/ObjectTemplates/DatasmithStaticMeshTemplate.cpp
@@ -123,7 +123,17 @@ void FDatasmithMeshSectionInfoMapTemplate::Apply( FMeshSectionInfoMap* Destinati
 {
 	for ( auto It = Map.CreateIterator(); It; ++It )
 	{
-		It->Value.Apply( &Destination->Map.FindOrAdd( It->Key ), PreviousTemplate ? PreviousTemplate->Map.Find( It->Key ) : nullptr );
+		FMeshSectionInfo* SectionInfo = Destination->Map.Find( It->Key );
+		if ( !SectionInfo )
+		{
+			//  If it wasn't found, must be added and the value forced on it
+			SectionInfo = &Destination->Map.Add( It->Key );
+			It->Value.Apply(SectionInfo, nullptr);
+		}
+		else
+		{
+			It->Value.Apply(SectionInfo, PreviousTemplate ? PreviousTemplate->Map.Find(It->Key) : nullptr);
+		}
 	}
 }
 
@@ -182,7 +192,14 @@ void UDatasmithStaticMeshTemplate::Apply( UObject* Destination, bool bForce )
 	DATASMITHOBJECTTEMPLATE_CONDITIONALSET( LightMapResolution, StaticMesh, PreviousStaticMeshTemplate );
 
 	// Section info map
-	SectionInfoMap.Apply( &StaticMesh->SectionInfoMap, PreviousStaticMeshTemplate ? &PreviousStaticMeshTemplate->SectionInfoMap : nullptr );
+	bool bResetSectionInfoMap = false;
+	if ( PreviousStaticMeshTemplate )
+	{
+		// If the number of sections is different, their order might be different (eg. from mesh editing) so the SectionInfoMap must be reset
+		bResetSectionInfoMap = PreviousStaticMeshTemplate->SectionInfoMap.Map.Num() != StaticMesh->SectionInfoMap.Map.Num();
+	}
+
+	SectionInfoMap.Apply(&StaticMesh->SectionInfoMap, !bResetSectionInfoMap && PreviousStaticMeshTemplate ? &PreviousStaticMeshTemplate->SectionInfoMap : nullptr);
 
 	// Build settings
 	for ( int32 SourceModelIndex = 0; SourceModelIndex < BuildSettings.Num(); ++SourceModelIndex )
@@ -207,16 +224,20 @@ void UDatasmithStaticMeshTemplate::Apply( UObject* Destination, bool bForce )
 	// Materials
 	for ( int32 StaticMaterialIndex = 0; StaticMaterialIndex < StaticMaterials.Num(); ++StaticMaterialIndex )
 	{
+		// If the SectionInfoMap was reset, the materials must be reapplied to follow it
+		bool bIgnorePreviousTemplate = bResetSectionInfoMap;
 		if ( !StaticMesh->StaticMaterials.IsValidIndex( StaticMaterialIndex ) )
 		{
+			// If it's a new added material, the value must be force applied to it by ignoring the previous template
 			StaticMesh->StaticMaterials.AddDefaulted();
+			bIgnorePreviousTemplate = true;
 		}
 
 		FStaticMaterial& StaticMaterial = StaticMesh->StaticMaterials[ StaticMaterialIndex ];
 
 		FDatasmithStaticMaterialTemplate* PreviousStaticMaterialTemplate = nullptr;
 
-		if ( PreviousStaticMeshTemplate && PreviousStaticMeshTemplate->StaticMaterials.IsValidIndex( StaticMaterialIndex ) )
+		if ( !bIgnorePreviousTemplate && PreviousStaticMeshTemplate && PreviousStaticMeshTemplate->StaticMaterials.IsValidIndex( StaticMaterialIndex ) )
 		{
 			PreviousStaticMaterialTemplate = &PreviousStaticMeshTemplate->StaticMaterials[ StaticMaterialIndex ];
 		}
diff --git a/Engine/Plugins/Runtime/EditableMesh/Source/EditableMesh/EditableStaticMeshAdapter.cpp b/Engine/Plugins/Runtime/EditableMesh/Source/EditableMesh/EditableStaticMeshAdapter.cpp
--- a/Engine/Plugins/Runtime/EditableMesh/Source/EditableMesh/EditableStaticMeshAdapter.cpp
+++ b/Engine/Plugins/Runtime/EditableMesh/Source/EditableMesh/EditableStaticMeshAdapter.cpp
@@ -647,9 +647,18 @@ void UEditableStaticMeshAdapter::OnRebuildRenderMesh( const UEditableMesh* Edita
 
 		for( const FPolygonGroupID PolygonGroupID : MeshDescription->PolygonGroups().GetElementIDs() )
 		{
-			FRenderingPolygonGroup& RenderingPolygonGroup = RenderingPolygonGroups[ PolygonGroupID ];
+			// Preserve the Sections order by finding the RenderingPolygonGroup with the RenderingSectionIndex that matches the PolygonGroupID
+			FPolygonGroupID RenderingGroupID = PolygonGroupID;
+			for ( const FPolygonGroupID RenderingPolygonGroupID : RenderingPolygonGroups.GetElementIDs() )
+			{
+				if ( RenderingPolygonGroups[ RenderingPolygonGroupID ].RenderingSectionIndex == PolygonGroupID.GetValue() )
+				{
+					RenderingGroupID = RenderingPolygonGroupID;
+					break;
+				}
+			}
 
-			RenderingPolygonGroup.RenderingSectionIndex = StaticMeshLOD.Sections.Num();
+			FRenderingPolygonGroup& RenderingPolygonGroup = RenderingPolygonGroups[ RenderingGroupID ];
 
 			// Create new rendering section
 			StaticMeshLOD.Sections.Add( FStaticMeshSection() );
@@ -1689,6 +1698,15 @@ void UEditableStaticMeshAdapter::OnCreatePolygonGroups( const UEditableMesh* Edi
 			StaticMeshSection.bEnableCollision = PolygonGroupCollision[ PolygonGroupID ];
 			StaticMeshSection.bCastShadow = PolygonGroupCastShadow[ PolygonGroupID ];
 			StaticMeshSection.MaterialIndex = MaterialIndex;
+
+#if WITH_EDITORONLY_DATA
+			// SectionInfoMap must be synced with the info of the new Section
+			FMeshSectionInfo Info;
+			Info.bEnableCollision = StaticMeshSection.bEnableCollision;
+			Info.bCastShadow = StaticMeshSection.bCastShadow;
+			Info.MaterialIndex = StaticMeshSection.MaterialIndex;
+			StaticMesh->SectionInfoMap.Set( StaticMeshLODIndex, LODSectionIndex, Info );
+#endif
 		}
 
 		// Insert the rendering polygon group for keeping track of these index buffer properties
@@ -1801,10 +1819,29 @@ void UEditableStaticMeshAdapter::OnDeletePolygonGroups( const UEditableMesh* Edi
 				if( StaticMeshSection.MaterialIndex > MaterialIndex )
 				{
 					StaticMeshSection.MaterialIndex--;
+
+#if WITH_EDITORONLY_DATA
+					// SectionInfoMap must be synced with the info of the modified Section
+					FMeshSectionInfo SectionInfo = StaticMesh->SectionInfoMap.Get( StaticMeshLODIndex, Index );
+					--SectionInfo.MaterialIndex;
+					StaticMesh->SectionInfoMap.Set( StaticMeshLODIndex, Index, SectionInfo );
+#endif
 				}
 			}
 
 			StaticMeshLOD.Sections.RemoveAt( RenderingSectionIndex );
+
+#if WITH_EDITORONLY_DATA
+			// SectionInfoMap must be re-indexed to account for the removed Section
+			uint32 NumSectionInfo = StaticMesh->SectionInfoMap.GetSectionNumber( StaticMeshLODIndex );
+			for ( uint32 Index = RenderingSectionIndex + 1; Index < NumSectionInfo; ++Index )
+			{
+				FMeshSectionInfo SectionInfo = StaticMesh->SectionInfoMap.Get( StaticMeshLODIndex, Index );
+				StaticMesh->SectionInfoMap.Set( StaticMeshLODIndex, Index - 1, SectionInfo );
+			}
+			// And remove the last SectionInfo from the map which is now invalid
+			StaticMesh->SectionInfoMap.Remove( StaticMeshLODIndex, NumSectionInfo - 1 );
+#endif
 		}
 
 		// Remove the rendering polygon group from the sparse array
