Missed on more win32 allignment fix.

diff --git a/Engine/Plugins/Runtime/EditableMesh/Source/EditableMesh/EditableMesh.cpp b/Engine/Plugins/Runtime/EditableMesh/Source/EditableMesh/EditableMesh.cpp
--- a/Engine/Plugins/Runtime/EditableMesh/Source/EditableMesh/EditableMesh.cpp
+++ b/Engine/Plugins/Runtime/EditableMesh/Source/EditableMesh/EditableMesh.cpp
@@ -8687,7 +8687,7 @@ void UEditableMesh::SearchSpatialDatabaseForPolygonsPotentiallyIntersectingLineS
 }
 
 // @todo mesheditor: Combine this function with the one above
-void UEditableMesh::SearchSpatialDatabaseForPolygonsPotentiallyIntersectingPlane( const FPlane InPlane, TArray<FPolygonID>& OutPolygons ) const
+void UEditableMesh::SearchSpatialDatabaseForPolygonsPotentiallyIntersectingPlane( const FPlane& InPlane, TArray<FPolygonID>& OutPolygons ) const
 {
 	OutPolygons.Reset();
 
diff --git a/Engine/Plugins/Runtime/EditableMesh/Source/EditableMesh/Public/EditableMesh.h b/Engine/Plugins/Runtime/EditableMesh/Source/EditableMesh/Public/EditableMesh.h
--- a/Engine/Plugins/Runtime/EditableMesh/Source/EditableMesh/Public/EditableMesh.h
+++ b/Engine/Plugins/Runtime/EditableMesh/Source/EditableMesh/Public/EditableMesh.h
@@ -439,7 +439,7 @@ public:
 	UFUNCTION( BlueprintPure, Category="Editable Mesh" ) void FindPolygonLoop( const FEdgeID EdgeID, TArray<FEdgeID>& OutEdgeLoopEdgeIDs, TArray<FEdgeID>& OutFlippedEdgeIDs, TArray<FEdgeID>& OutReversedEdgeIDPathToTake, TArray<FPolygonID>& OutPolygonIDsToSplit ) const;
 	UFUNCTION( BlueprintPure, Category="Editable Mesh" ) void SearchSpatialDatabaseForPolygonsPotentiallyIntersectingLineSegment( const FVector LineSegmentStart, const FVector LineSegmentEnd, TArray<FPolygonID>& OutPolygons ) const;
 	UFUNCTION( BlueprintPure, Category="Editable Mesh" ) void SearchSpatialDatabaseForPolygonsInVolume( const TArray<FPlane>& Planes, TArray<FPolygonID>& OutPolygons ) const;
-	UFUNCTION( BlueprintPure, Category="Editable Mesh" ) void SearchSpatialDatabaseForPolygonsPotentiallyIntersectingPlane( const FPlane InPlane, TArray<FPolygonID>& OutPolygons ) const;
+	UFUNCTION( BlueprintPure, Category="Editable Mesh" ) void SearchSpatialDatabaseForPolygonsPotentiallyIntersectingPlane( const FPlane& InPlane, TArray<FPolygonID>& OutPolygons ) const;
 
 
 	UFUNCTION( BlueprintCallable, Category="Editable Mesh" ) void SetSubdivisionCount( const int32 NewSubdivisionCount );
