Corrected type of PolygonGroup::EnableCollision and PolygonGroup::CastShadow attributes.

diff --git a/Engine/Plugins/Runtime/EditableMesh/Source/EditableMesh/StaticMeshEditableMeshFormat.cpp b/Engine/Plugins/Runtime/EditableMesh/Source/EditableMesh/StaticMeshEditableMeshFormat.cpp
--- a/Engine/Plugins/Runtime/EditableMesh/Source/EditableMesh/StaticMeshEditableMeshFormat.cpp
+++ b/Engine/Plugins/Runtime/EditableMesh/Source/EditableMesh/StaticMeshEditableMeshFormat.cpp
@@ -55,8 +55,8 @@ UEditableMesh* FStaticMeshEditableMeshFormat::MakeEditableMesh( UPrimitiveCompon
 	MeshDescription->PolygonAttributes().RegisterAttribute<FVector>( MeshAttribute::Polygon::Binormal, 1, FVector::ZeroVector, EMeshAttributeFlags::Transient );
 	MeshDescription->PolygonAttributes().RegisterAttribute<FVector>( MeshAttribute::Polygon::Center, 1, FVector::ZeroVector, EMeshAttributeFlags::Transient );
 	MeshDescription->PolygonGroupAttributes().RegisterAttribute<FName>( MeshAttribute::PolygonGroup::MaterialAssetName );
-	MeshDescription->PolygonGroupAttributes().RegisterAttribute<FName>( MeshAttribute::PolygonGroup::EnableCollision );
-	MeshDescription->PolygonGroupAttributes().RegisterAttribute<FName>( MeshAttribute::PolygonGroup::CastShadow );
+	MeshDescription->PolygonGroupAttributes().RegisterAttribute<bool>( MeshAttribute::PolygonGroup::EnableCollision );
+	MeshDescription->PolygonGroupAttributes().RegisterAttribute<bool>( MeshAttribute::PolygonGroup::CastShadow );
 
 	UEditableStaticMeshAdapter* EditableStaticMesh = NewObject<UEditableStaticMeshAdapter>( EditableMesh );
 	EditableMesh->Adapters.Add( EditableStaticMesh );
