Remove all references to holes-in-polygons from MeshDescription and EditableMesh

diff --git a/Engine/Plugins/Runtime/EditableMesh/Source/EditableMesh/EditableMesh.cpp b/Engine/Plugins/Runtime/EditableMesh/Source/EditableMesh/EditableMesh.cpp
--- a/Engine/Plugins/Runtime/EditableMesh/Source/EditableMesh/EditableMesh.cpp
+++ b/Engine/Plugins/Runtime/EditableMesh/Source/EditableMesh/EditableMesh.cpp
@@ -1004,32 +1004,6 @@ FVertexID UEditableMesh::GetPolygonPerimeterVertex( const FPolygonID PolygonID,
 }
 
 
-int32 UEditableMesh::GetPolygonHoleCount( const FPolygonID PolygonID ) const
-{
-	return GetMeshDescription()->GetNumPolygonHoles( PolygonID );
-}
-
-
-int32 UEditableMesh::GetPolygonHoleVertexCount( const FPolygonID PolygonID, const int32 HoleNumber ) const
-{
-	return GetMeshDescription()->GetPolygonHoleVertexInstances( PolygonID, HoleNumber ).Num();
-}
-
-
-FVertexInstanceID UEditableMesh::GetPolygonHoleVertexInstance( const FPolygonID PolygonID, const int32 HoleNumber, const int32 PolygonVertexNumber ) const
-{
-	const TArray<FVertexInstanceID>& VertexInstanceIDs = GetMeshDescription()->GetPolygonHoleVertexInstances( PolygonID, HoleNumber );
-	return ( PolygonVertexNumber >= 0 && PolygonVertexNumber < VertexInstanceIDs.Num() ) ? VertexInstanceIDs[ PolygonVertexNumber ] : FVertexInstanceID::Invalid;
-}
-
-
-FVertexID UEditableMesh::GetPolygonHoleVertex( const FPolygonID PolygonID, const int32 HoleNumber, const int32 PolygonVertexNumber ) const
-{
-	const FVertexInstanceID VertexInstanceID = GetPolygonHoleVertexInstance( PolygonID, HoleNumber, PolygonVertexNumber );
-	return ( VertexInstanceID != FVertexInstanceID::Invalid ) ? GetMeshDescription()->GetVertexInstanceVertex( VertexInstanceID ) : FVertexID::Invalid;
-}
-
-
 int32 UEditableMesh::GetPolygonTriangulatedTriangleCount( const FPolygonID PolygonID ) const
 {
 	return GetMeshDescription()->GetPolygonTriangles( PolygonID ).Num();
@@ -1255,26 +1229,6 @@ void UEditableMesh::ReplaceVertexInstanceInPolygons( const FVertexInstanceID Old
 			ChangeVertexInstances.PerimeterVertexIndicesAndInstanceIDs[ 0 ].ContourIndex = VertexInstanceIndex;
 			ChangeVertexInstances.PerimeterVertexIndicesAndInstanceIDs[ 0 ].VertexInstanceID = NewVertexInstanceID;
 		}
-		else
-		{
-			ChangeVertexInstances.VertexIndicesAndInstanceIDsForEachHole.SetNum( Polygon.HoleContours.Num() );
-
-			int32 HoleIndex = 0;
-			for( const FMeshPolygonContour& HoleContour : Polygon.HoleContours )
-			{
-				VertexInstanceIndex = HoleContour.VertexInstanceIDs.Find( OldVertexInstanceID );
-				if( VertexInstanceIndex != INDEX_NONE )
-				{
-					// Found the vertex instance in a hole
-					ChangeVertexInstances.VertexIndicesAndInstanceIDsForEachHole[ HoleIndex ].VertexIndicesAndInstanceIDs.Emplace();
-					ChangeVertexInstances.VertexIndicesAndInstanceIDsForEachHole[ HoleIndex ].VertexIndicesAndInstanceIDs[ 0 ].ContourIndex = VertexInstanceIndex;
-					ChangeVertexInstances.VertexIndicesAndInstanceIDsForEachHole[ HoleIndex ].VertexIndicesAndInstanceIDs[ 0 ].VertexInstanceID = NewVertexInstanceID;
-					break;
-				}
-
-				HoleIndex++;
-			}
-		}
 
 		// We expect to have found the vertex instance somewhere in one of the polygon contours
 		check( VertexInstanceIndex != INDEX_NONE );
@@ -1598,13 +1552,6 @@ int32 UEditableMesh::GetPolygonPerimeterEdgeCount( const FPolygonID PolygonID )
 }
 
 
-int32 UEditableMesh::GetPolygonHoleEdgeCount( const FPolygonID PolygonID, const int32 HoleNumber ) const
-{
-	// All polygons have the same number of edges as they do vertices
-	return GetPolygonHoleVertexCount( PolygonID, HoleNumber );
-}
-
-
 void UEditableMesh::GetPolygonPerimeterVertices( const FPolygonID PolygonID, TArray<FVertexID>& OutPolygonPerimeterVertexIDs ) const
 {
 	GetMeshDescription()->GetPolygonPerimeterVertices( PolygonID, OutPolygonPerimeterVertexIDs );
@@ -1617,26 +1564,6 @@ void UEditableMesh::GetPolygonPerimeterVertexInstances( const FPolygonID Polygon
 }
 
 
-void UEditableMesh::GetPolygonHoleVertices( const FPolygonID PolygonID, const int32 HoleNumber, TArray<FVertexID>& OutHoleVertexIDs ) const
-{
-	OutHoleVertexIDs.Reset();	// in case it fails below
-	if( HoleNumber < GetMeshDescription()->GetNumPolygonHoles( PolygonID ) )
-	{
-		GetMeshDescription()->GetPolygonHoleVertices( PolygonID, HoleNumber, OutHoleVertexIDs );
-	}
-}
-
-
-void UEditableMesh::GetPolygonHoleVertexInstances( const FPolygonID PolygonID, const int32 HoleNumber, TArray<FVertexInstanceID>& OutHoleVertexInstanceIDs ) const
-{
-	OutHoleVertexInstanceIDs.Reset();	// in case it fails below
-	if( HoleNumber < GetMeshDescription()->GetNumPolygonHoles( PolygonID ) )
-	{
-		OutHoleVertexInstanceIDs = GetMeshDescription()->GetPolygonHoleVertexInstances( PolygonID, HoleNumber );
-	}
-}
-
-
 FEdgeID UEditableMesh::GetPolygonPerimeterEdge( const FPolygonID PolygonID, const int32 PerimeterEdgeNumber, bool& bOutEdgeWindingIsReversedForPolygon ) const
 {
 	const FMeshDescription* Description = GetMeshDescription();
@@ -1665,27 +1592,6 @@ FEdgeID UEditableMesh::GetVertexPairEdge( const FVertexID StartVertexID, const F
 }
 
 
-FEdgeID UEditableMesh::GetPolygonHoleEdge( const FPolygonID PolygonID, const int32 HoleNumber, const int32 HoleEdgeNumber ) const
-{
-	const FMeshDescription* Description = GetMeshDescription();
-
-	if( HoleNumber >= Description->GetNumPolygonHoles( PolygonID ) )
-	{
-		return FEdgeID::Invalid;
-	}
-
-	const TArray<FVertexInstanceID>& VertexInstanceIDs = Description->GetPolygonHoleVertexInstances( PolygonID, HoleNumber );
-	if( HoleEdgeNumber >= VertexInstanceIDs.Num() )
-	{
-		return FEdgeID::Invalid;
-	}
-
-	const FVertexID VertexID0 = Description->GetVertexInstanceVertex( VertexInstanceIDs[ HoleEdgeNumber ] );
-	const FVertexID VertexID1 = Description->GetVertexInstanceVertex( VertexInstanceIDs[ ( HoleEdgeNumber + 1 ) % VertexInstanceIDs.Num() ] );
-	return Description->GetVertexPairEdge( VertexID0, VertexID1 );
-}
-
-
 void UEditableMesh::GetPolygonPerimeterEdges( const FPolygonID PolygonID, TArray<FEdgeID>& OutPolygonPerimeterEdgeIDs ) const
 {
 	const FMeshDescription* Description = GetMeshDescription();
@@ -1703,28 +1609,6 @@ void UEditableMesh::GetPolygonPerimeterEdges( const FPolygonID PolygonID, TArray
 }
 
 
-void UEditableMesh::GetPolygonHoleEdges( const FPolygonID PolygonID, const int32 HoleNumber, TArray<FEdgeID>& OutHoleEdgeIDs ) const
-{
-	const FMeshDescription* Description = GetMeshDescription();
-
-	OutHoleEdgeIDs.Reset();
-	if( HoleNumber < Description->GetNumPolygonHoles( PolygonID ) )
-	{
-		const TArray<FVertexInstanceID>& VertexInstanceIDs = Description->GetPolygonHoleVertexInstances( PolygonID, HoleNumber );
-		const int32 NumContourEdges = VertexInstanceIDs.Num();
-
-		OutHoleEdgeIDs.SetNumUninitialized( NumContourEdges, false );
-		for( int32 Index = 0; Index < NumContourEdges; ++Index )
-		{
-			const int32 NextIndex = ( Index + 1 == NumContourEdges ) ? 0 : Index + 1;
-			const FVertexID VertexID0 = Description->GetVertexInstanceVertex( VertexInstanceIDs[ Index ] );
-			const FVertexID VertexID1 = Description->GetVertexInstanceVertex( VertexInstanceIDs[ NextIndex ] );
-			OutHoleEdgeIDs[ Index ] = Description->GetVertexPairEdge( VertexID0, VertexID1 );
-		}
-	}
-}
-
-
 void UEditableMesh::GetPolygonAdjacentPolygons( const FPolygonID PolygonID, TArray<FPolygonID>& OutAdjacentPolygons ) const
 {
 	OutAdjacentPolygons.Reset();
@@ -1971,7 +1855,7 @@ void UEditableMesh::RefreshOpenSubdiv()
 			OsdTopologyDescriptor.cornerVertexIndices = OsdCornerVertexIndices.GetData();
 			OsdTopologyDescriptor.cornerWeights = OsdCornerWeights.GetData();
 
-			OsdTopologyDescriptor.numHoles = 0;	// @todo mesheditor holes
+			OsdTopologyDescriptor.numHoles = 0;
 			OsdTopologyDescriptor.holeIndices = nullptr;
 
 			OsdTopologyDescriptor.isLeftHanded = true;
@@ -2539,7 +2423,6 @@ void UEditableMesh::ComputePolygonTriangulation( const FPolygonID PolygonID, TAr
 
 	OutTriangles.Reset();
 
-	// @todo mesheditor holes: Does not support triangles with holes yet!
 	// @todo mesheditor: Perhaps should always attempt to triangulate by splitting polygons along the shortest edge, for better determinism.
 
 //	const FMeshPolygon& Polygon = GetMeshDescription()->GetPolygon( PolygonID );
@@ -2837,71 +2720,6 @@ void UEditableMesh::CreateMissingPolygonPerimeterEdges( const FPolygonID Polygon
 }
 
 
-void UEditableMesh::CreateMissingPolygonHoleEdges( const FPolygonID PolygonID, const int32 HoleNumber, TArray<FEdgeID>& OutNewEdgeIDs )
-{
-	// @todo mesheditor: Orphaned method - should we keep it or is it entirely unnecessary now?
-
-	OutNewEdgeIDs.Reset();
-
-	const int32 NumHoleEdges = GetPolygonHoleEdgeCount( PolygonID, HoleNumber );
-	const int32 NumHoleVertices = NumHoleEdges;		// Edge and vertex count are always the same
-
-	for( int32 HoleEdgeNumber = 0; HoleEdgeNumber < NumHoleEdges; ++HoleEdgeNumber )
-	{
-		const int32 HoleVertexNumber = HoleEdgeNumber;	// Edge and vertex counts are always the same
-
-		const FVertexID VertexID = GetPolygonHoleVertex( PolygonID, HoleNumber, HoleVertexNumber );
-		const FVertexID NextVertexID = GetPolygonHoleVertex( PolygonID, HoleNumber, ( HoleVertexNumber + 1 ) % NumHoleVertices );
-
-		// Find the edge that connects these vertices
-		FEdgeID FoundEdgeID = FEdgeID::Invalid;
-		{
-			bool bFoundEdge = false;
-
-			const int32 NumVertexConnectedEdges = GetVertexConnectedEdgeCount( VertexID );	// @todo mesheditor perf: Could be made faster by always iterating over the vertex with the fewest number of connected edges, instead of the first edge vertex
-			for( int32 VertexEdgeNumber = 0; VertexEdgeNumber < NumVertexConnectedEdges; ++VertexEdgeNumber )
-			{
-				const FEdgeID VertexConnectedEdgeID = GetVertexConnectedEdge( VertexID, VertexEdgeNumber );
-
-				// Try the edge's first vertex.  Does it point to us?
-				FVertexID OtherEdgeVertexID = GetEdgeVertex( VertexConnectedEdgeID, 0 );
-				if( OtherEdgeVertexID == VertexID )
-				{
-					// Must be the the other one
-					OtherEdgeVertexID = GetEdgeVertex( VertexConnectedEdgeID, 1 );
-				}
-
-				if( OtherEdgeVertexID == NextVertexID )
-				{
-					// We found the edge!
-					FoundEdgeID = VertexConnectedEdgeID;
-					bFoundEdge = true;
-					break;
-				}
-			}
-
-			if( !bFoundEdge )
-			{
-				// Create the new edge!  Note that this does not connect the edge to the polygon.  We expect the caller to do that afterwards.
-				FEdgeToCreate EdgeToCreate;
-				EdgeToCreate.VertexID0 = VertexID;
-				EdgeToCreate.VertexID1 = NextVertexID;
-
-				static TArray< FEdgeToCreate > EdgesToCreate;
-				EdgesToCreate.Reset();
-				EdgesToCreate.Add( EdgeToCreate );
-
-				static TArray< FEdgeID > NewEdgeIDs;
-				NewEdgeIDs.Reset();
-				CreateEdges( EdgesToCreate, /* Out */ NewEdgeIDs );
-
-				OutNewEdgeIDs.Append( NewEdgeIDs );
-			}
-		}
-	}
-}
-
-
 void UEditableMesh::SplitEdge( const FEdgeID EdgeID, const TArray<float>& Splits, TArray<FVertexID>& OutNewVertexIDs )
 {
 	EM_ENTER( TEXT( "SplitEdge: %s %s" ), *EdgeID.ToString(), *LogHelpers::ArrayToString( Splits ) );
@@ -4328,26 +4146,14 @@ void UEditableMesh::CreatePolygons( const TArray<FPolygonToCreate>& PolygonsToCr
 			ExistingEdgeIDs.Append( ExistingEdgeIDsForContour );
 			OutNewEdgeIDs.Append( NewEdgeIDsForContour );
 
-			static TArray<TArray<FMeshDescription::FContourPoint>> HoleContourPoints;
-
-			// Iterate through hole definitions, and create edges and vertex instances for those
-			for( const FPolygonHoleVertices& PolygonHole : PolygonToCreate.PolygonHoles )
-			{
-				HoleContourPoints.Emplace();
-
-				CreatePolygonContour( PolygonHole.HoleVertices, ExistingEdgeIDsForContour, NewEdgeIDsForContour, HoleContourPoints.Last() );
-				ExistingEdgeIDs.Append( ExistingEdgeIDsForContour );
-				OutNewEdgeIDs.Append( NewEdgeIDsForContour );
-			}
-
 			FPolygonID PolygonID = PolygonToCreate.OriginalPolygonID;
 			if( PolygonID != FPolygonID::Invalid )
 			{
-				GetMeshDescription()->CreatePolygonWithID( PolygonID, PolygonToCreate.PolygonGroupID, PerimeterContourPoints, HoleContourPoints );
+				GetMeshDescription()->CreatePolygonWithID( PolygonID, PolygonToCreate.PolygonGroupID, PerimeterContourPoints );
 			}
 			else
 			{
-				PolygonID = GetMeshDescription()->CreatePolygon( PolygonToCreate.PolygonGroupID, PerimeterContourPoints, HoleContourPoints );
+				PolygonID = GetMeshDescription()->CreatePolygon( PolygonToCreate.PolygonGroupID, PerimeterContourPoints );
 			}
 
 			OutNewPolygonIDs.Add( PolygonID );
@@ -4462,13 +4268,6 @@ void UEditableMesh::DeletePolygons( const TArray<FPolygonID>& PolygonIDsToDelete
 			PolygonToCreate.OriginalPolygonID = PolygonID;
 
 			BackupPolygonContour( Polygon.PerimeterContour, PolygonToCreate.PerimeterVertices );
-			
-			PolygonToCreate.PolygonHoles.Reserve( Polygon.HoleContours.Num() );
-			for( const FMeshPolygonContour& HoleContour : Polygon.HoleContours )
-			{
-				PolygonToCreate.PolygonHoles.Emplace();
-				BackupPolygonContour( HoleContour, PolygonToCreate.PolygonHoles.Last().HoleVertices );
-			}
 		}
 
 		AddUndo( MakeUnique<FCreatePolygonsChange>( MoveTemp( RevertInput ) ) );
@@ -4780,23 +4579,6 @@ void UEditableMesh::ChangePolygonsVertexInstances( const TArray<FChangeVertexIns
 				RevertIndexAndInstance.ContourIndex = IndexAndInstance.ContourIndex;
 				RevertIndexAndInstance.VertexInstanceID = Polygon.PerimeterContour.VertexInstanceIDs[ IndexAndInstance.ContourIndex ];
 			}
-
-			int32 HoleIndex = 0;
-			for( const FVertexInstancesForPolygonHole& VertexInstancesForPolygonHole : VertexInstancesForPolygon.VertexIndicesAndInstanceIDsForEachHole )
-			{
-				RevertVertexInstancesForPolygon.VertexIndicesAndInstanceIDsForEachHole.Emplace();
-				FVertexInstancesForPolygonHole& RevertVertexInstancesForPolygonHole = RevertVertexInstancesForPolygon.VertexIndicesAndInstanceIDsForEachHole.Last();
-
-				for( const FVertexIndexAndInstanceID& IndexAndInstance : VertexInstancesForPolygonHole.VertexIndicesAndInstanceIDs )
-				{
-					RevertVertexInstancesForPolygonHole.VertexIndicesAndInstanceIDs.Emplace();
-					FVertexIndexAndInstanceID& RevertIndexAndInstance = RevertVertexInstancesForPolygonHole.VertexIndicesAndInstanceIDs.Last();
-					RevertIndexAndInstance.ContourIndex = IndexAndInstance.ContourIndex;
-					RevertIndexAndInstance.VertexInstanceID = Polygon.HoleContours[ HoleIndex ].VertexInstanceIDs[ IndexAndInstance.ContourIndex ];
-				}
-
-				HoleIndex++;
-			}
 		}
 
 		AddUndo( MakeUnique<FChangePolygonsVertexInstancesChange>( MoveTemp( RevertInput ) ) );
@@ -4838,38 +4620,6 @@ void UEditableMesh::ChangePolygonsVertexInstances( const TArray<FChangeVertexIns
 					}
 				}
 			}
-
-			int32 HoleIndex = 0;
-			for( const FVertexInstancesForPolygonHole& VertexInstancesForPolygonHole : VertexInstancesForPolygon.VertexIndicesAndInstanceIDsForEachHole )
-			{
-				for( const FVertexIndexAndInstanceID& IndexAndInstance : VertexInstancesForPolygonHole.VertexIndicesAndInstanceIDs )
-				{
-					// Disconnect old vertex instance from polygon, and connect new one
-					const FVertexInstanceID OldVertexInstanceID = Polygon.HoleContours[ HoleIndex ].VertexInstanceIDs[ IndexAndInstance.ContourIndex ];
-					FMeshVertexInstance& OldVertexInstance = VertexInstances[ OldVertexInstanceID ];
-					verify( OldVertexInstance.ConnectedPolygons.Remove( PolygonID ) == 1 );
-
-					FMeshVertexInstance& NewVertexInstance = VertexInstances[ IndexAndInstance.VertexInstanceID ];
-					check( !NewVertexInstance.ConnectedPolygons.Contains( PolygonID ) );
-					NewVertexInstance.ConnectedPolygons.Add( PolygonID );
-
-					Polygon.HoleContours[ HoleIndex ].VertexInstanceIDs[ IndexAndInstance.ContourIndex ] = IndexAndInstance.VertexInstanceID;
-
-					// Fix up triangle list
-					for( FMeshTriangle& Triangle : Polygon.Triangles )
-					{
-						for( int32 VertexIndex = 0; VertexIndex < 3; ++VertexIndex )
-						{
-							if( Triangle.GetVertexInstanceID( VertexIndex ) == OldVertexInstanceID )
-							{
-								Triangle.SetVertexInstanceID( VertexIndex, IndexAndInstance.VertexInstanceID );
-							}
-						}
-					}
-				}
-
-				HoleIndex++;
-			}
 		}
 	}
 
@@ -4926,19 +4676,14 @@ void UEditableMesh::SetPolygonsVertexAttributes( const TArray<FVertexAttributesF
 		const FPolygonID PolygonID = VertexAttributesForPolygon.PolygonID;
 		FMeshPolygon& Polygon = Polygons[ PolygonID ];
 
-		SetPolygonContourVertexAttributes( Polygon.PerimeterContour, PolygonID, INDEX_NONE, VertexAttributesForPolygon.PerimeterVertexAttributeLists );
-
-		for( int32 HoleIndex = 0; HoleIndex < VertexAttributesForPolygon.VertexAttributeListsForEachHole.Num(); ++HoleIndex )
-		{
-			SetPolygonContourVertexAttributes( Polygon.HoleContours[ HoleIndex ], PolygonID, HoleIndex, VertexAttributesForPolygon.VertexAttributeListsForEachHole[ HoleIndex ].VertexAttributeList );
-		}
+		SetPolygonContourVertexAttributes( Polygon.PerimeterContour, PolygonID, VertexAttributesForPolygon.PerimeterVertexAttributeLists );
 	}
 
 	EM_EXIT( TEXT( "SetPolygonsVertexAttributes returned" ) );
 }
 
 
-void UEditableMesh::SetPolygonContourVertexAttributes( FMeshPolygonContour& Contour, const FPolygonID PolygonID, const int32 HoleIndex, const TArray<FMeshElementAttributeList>& AttributeLists )
+void UEditableMesh::SetPolygonContourVertexAttributes( FMeshPolygonContour& Contour, const FPolygonID PolygonID, const TArray<FMeshElementAttributeList>& AttributeLists )
 {
 	const int32 NumContourVertices = Contour.VertexInstanceIDs.Num();
 	check( AttributeLists.Num() == NumContourVertices );
@@ -4993,19 +4738,10 @@ void UEditableMesh::SetPolygonContourVertexAttributes( FMeshPolygonContour& Cont
 			VertexInstancesToChange.SetNum( 1 );
 
 			VertexInstancesToChange[ 0 ].PolygonID = PolygonID;
-			if( HoleIndex == INDEX_NONE )
-			{
-				VertexInstancesToChange[ 0 ].PerimeterVertexIndicesAndInstanceIDs.Emplace();
-				VertexInstancesToChange[ 0 ].PerimeterVertexIndicesAndInstanceIDs[ 0 ].ContourIndex = Index;
-				VertexInstancesToChange[ 0 ].PerimeterVertexIndicesAndInstanceIDs[ 0 ].VertexInstanceID = NewVertexInstanceIDs[ 0 ];
-			}
-			else
-			{
-				VertexInstancesToChange[ 0 ].VertexIndicesAndInstanceIDsForEachHole.SetNum( HoleIndex + 1 );
-				VertexInstancesToChange[ 0 ].VertexIndicesAndInstanceIDsForEachHole[ HoleIndex ].VertexIndicesAndInstanceIDs.Emplace();
-				VertexInstancesToChange[ 0 ].VertexIndicesAndInstanceIDsForEachHole[ HoleIndex ].VertexIndicesAndInstanceIDs[ 0 ].ContourIndex = Index;
-				VertexInstancesToChange[ 0 ].VertexIndicesAndInstanceIDsForEachHole[ HoleIndex ].VertexIndicesAndInstanceIDs[ 0 ].VertexInstanceID = NewVertexInstanceIDs[ 0 ];
-			}
+
+			VertexInstancesToChange[ 0 ].PerimeterVertexIndicesAndInstanceIDs.Emplace();
+			VertexInstancesToChange[ 0 ].PerimeterVertexIndicesAndInstanceIDs[ 0 ].ContourIndex = Index;
+			VertexInstancesToChange[ 0 ].PerimeterVertexIndicesAndInstanceIDs[ 0 ].VertexInstanceID = NewVertexInstanceIDs[ 0 ];
 
 			ChangePolygonsVertexInstances( VertexInstancesToChange );
 
@@ -5175,7 +4911,6 @@ void UEditableMesh::TryToRemovePolygonEdge( const FEdgeID EdgeID, bool& bOutWasE
 					FPolygonToCreate& PolygonToCreate = *new( PolygonsToCreate ) FPolygonToCreate();
 					PolygonToCreate.PolygonGroupID = PolygonGroupID;
 					PolygonToCreate.PerimeterVertices = NewPolygonVertices;	// @todo mesheditor perf: Extra allocatons/copies: Ideally MoveTemp() here but we can't move a STATIC local!
-					// @todo mesheditor hole: Hole support needed!
 
 					static TArray<FPolygonID> NewPolygonIDs;
 					NewPolygonIDs.Reset();
@@ -5237,7 +4972,6 @@ void UEditableMesh::TryToRemoveVertex( const FVertexID VertexID, bool& bOutWasVe
 		{
 			for( const FPolygonID PolygonID : NewEdgeConnectedPolygons )
 			{
-				// @todo mesheditor holes
 				const int32 PolygonVertexNumber = FindPolygonPerimeterVertexNumberForVertex( PolygonID, VertexID );
 				check( PolygonVertexNumber != INDEX_NONE );
 				const bool bDeleteOrphanedVertexInstances = false;
@@ -5557,8 +5291,6 @@ void UEditableMesh::ExtrudePolygons( const TArray<FPolygonID>& PolygonIDs, const
 							PolygonToCreate.PerimeterVertices = NewSidePolygonVertices;	// @todo mesheditor perf: Copying static array here, ideally allocations could be avoided
 							PolygonToCreate.PolygonEdgeHardness = EPolygonEdgeHardness::AllEdgesHard;
 
-							// NOTE: We never create holes in side polygons
-
 							static TArray<FPolygonID> NewPolygonIDs;
 							NewPolygonIDs.Reset();
 							static TArray<FEdgeID> NewEdgeIDs;
@@ -5620,7 +5352,6 @@ void UEditableMesh::ExtrudePolygons( const TArray<FPolygonID>& PolygonIDs, const
 				PolygonToCreate.PolygonGroupID = PolygonGroupID;
 				PolygonToCreate.PolygonEdgeHardness = EPolygonEdgeHardness::AllEdgesHard;
 				PolygonToCreate.PerimeterVertices = NewFrontPolygonVertices;	// @todo mesheditor perf: Copying static array here, ideally allocations could be avoided
-																				// @todo mesheditor holes: This algorithm is ignoring polygon holes!  Should be easy to support this by doing roughly the same thing.
 				static TArray<FPolygonID> NewPolygonIDs;
 				NewPolygonIDs.Reset();
 				static TArray<FEdgeID> NewEdgeIDs;
@@ -6349,19 +6080,6 @@ float UEditableMesh::GetPolygonCornerAngleForVertex( const FPolygonID PolygonID,
 		// Return the internal angle if found
 		return GetContourAngle( Polygon.PerimeterContour, ContourIndex );
 	}
-	else
-	{
-		// If not found, look in all the holes
-		for( const FMeshPolygonContour& HoleContour : Polygon.HoleContours )
-		{
-			ContourIndex = HoleContour.VertexInstanceIDs.IndexOfByPredicate( IsVertexInstancedFromThisVertex );
-			if( ContourIndex != INDEX_NONE )
-			{
-				// Hole vertex contribution is the part which ISN'T the internal angle of the contour, so subtract from 2*pi
-				return ( 2.0f * PI ) - GetContourAngle( HoleContour, ContourIndex );
-			}
-		}
-	}
 
 	// Found nothing; return 0
 	return 0.0f;
@@ -6442,7 +6160,6 @@ void UEditableMesh::GenerateTangentsAndNormals()
 	for( const FPolygonID PolygonID : PolygonsPendingNewTangentBasis )
 	{
 		VertexInstanceIDs.Append( GetMeshDescription()->GetPolygonPerimeterVertexInstances( PolygonID ) );
-		// @todo holes
 	}
 
 	static TArray<FAttributesForVertexInstance> AttributesForVertexInstances;
@@ -6743,7 +6460,6 @@ void UEditableMesh::GenerateTangentsForPolygons( const TArray< FPolygonID >& Pol
 
 			const FPolygonID PolygonID = UserData.Polygons[ MikkFaceIndex ];
 
-			// @todo mesheditor holes: Can we even use Mikktpsace for polygons with holes?  Do we need to run the triangulated triangles through instead?
 			const int32 PolygonVertexCount = UserData.Self->GetPolygonPerimeterVertexCount( PolygonID );
 			return PolygonVertexCount;
 		}
@@ -7152,24 +6868,6 @@ int32 UEditableMesh::FindPolygonPerimeterVertexNumberForVertex( const FPolygonID
 }
 
 
-int32 UEditableMesh::FindPolygonHoleVertexNumberForVertex( const FPolygonID PolygonID, const int32 HoleNumber, const FVertexID VertexID ) const
-{
-	int32 FoundPolygonVertexNumber = INDEX_NONE;
-
-	const int32 PolygonHoleVertexCount = this->GetPolygonHoleVertexCount( PolygonID, HoleNumber );
-	for( int32 PolygonVertexNumber = 0; PolygonVertexNumber < PolygonHoleVertexCount; ++PolygonVertexNumber )
-	{
-		if( VertexID == this->GetPolygonHoleVertex( PolygonID, HoleNumber, PolygonVertexNumber ) )
-		{
-			FoundPolygonVertexNumber = PolygonVertexNumber;
-			break;
-		}
-	}
-
-	return FoundPolygonVertexNumber;
-}
-
-
 int32 UEditableMesh::FindPolygonPerimeterEdgeNumberForVertices( const FPolygonID PolygonID, const FVertexID EdgeVertexID0, const FVertexID EdgeVertexID1 ) const
 {
 	// @todo mesheditor: surely quicker just to iterate round the perimeter, checking pairs of adjacent vertices against the passed IDs (in both orders)?
@@ -7198,34 +6896,6 @@ int32 UEditableMesh::FindPolygonPerimeterEdgeNumberForVertices( const FPolygonID
 }
 
 
-int32 UEditableMesh::FindPolygonHoleEdgeNumberForVertices( const FPolygonID PolygonID, const int32 HoleNumber, const FVertexID EdgeVertexID0, const FVertexID EdgeVertexID1 ) const
-{
-	// @todo mesheditor: surely quicker just to iterate round the perimeter, checking pairs of adjacent vertices against the passed IDs (in both orders)?
-
-	int32 FoundPolygonEdgeNumber = INDEX_NONE;
-
-	static TArray<FEdgeID> EdgeIDs;
-	GetPolygonHoleEdges( PolygonID, HoleNumber, /* Out */ EdgeIDs );
-
-	for( int32 PolygonEdgeNumber = 0; PolygonEdgeNumber < EdgeIDs.Num(); ++PolygonEdgeNumber )
-	{
-		const FEdgeID EdgeID = EdgeIDs[ PolygonEdgeNumber ];
-
-		FVertexID TestEdgeVertexIDs[ 2 ];
-		GetEdgeVertices( EdgeID, /* Out */ TestEdgeVertexIDs[0], /* Out */ TestEdgeVertexIDs[1] );
-
-		if( ( TestEdgeVertexIDs[ 0 ] == EdgeVertexID0 && TestEdgeVertexIDs[ 1 ] == EdgeVertexID1 ) ||
-			( TestEdgeVertexIDs[ 1 ] == EdgeVertexID0 && TestEdgeVertexIDs[ 0 ] == EdgeVertexID1 ) )
-		{
-			FoundPolygonEdgeNumber = PolygonEdgeNumber;
-			break;
-		}
-	}
-
-	return FoundPolygonEdgeNumber;
-}
-
-
 void UEditableMesh::FlipPolygons( const TArray<FPolygonID>& PolygonIDs )
 {
 	EM_ENTER( TEXT( "FlipPolygons %s" ), *LogHelpers::ArrayToString( PolygonIDs ) );
@@ -7257,8 +6927,6 @@ void UEditableMesh::FlipPolygons( const TArray<FPolygonID>& PolygonIDs )
 			// Just create a polygon from existing vertex instance IDs
 			PerimeterVertex.VertexInstanceID = PerimeterVertexInstanceIDs[ VertexNumber ];
 		}
-
-		// @todo mesheditor hole: Perserve holes in this polygon!  (Vertex IDs and Attributes!) Should these be flipped too?  Probably.
 	}
 
 	// Delete all of the polygons
diff --git a/Engine/Plugins/Runtime/EditableMesh/Source/EditableMesh/Public/EditableMesh.h b/Engine/Plugins/Runtime/EditableMesh/Source/EditableMesh/Public/EditableMesh.h
--- a/Engine/Plugins/Runtime/EditableMesh/Source/EditableMesh/Public/EditableMesh.h
+++ b/Engine/Plugins/Runtime/EditableMesh/Source/EditableMesh/Public/EditableMesh.h
@@ -212,22 +212,6 @@ public:
 	UFUNCTION( BlueprintPure, Category="Editable Mesh" )
 	FVertexInstanceID GetPolygonPerimeterVertexInstance( const FPolygonID PolygonID, const int32 PolygonVertexNumber ) const;
 
-	/** Returns the number of hole contours this polygon has */
-	UFUNCTION( BlueprintPure, Category="Editable Mesh" )
-	int32 GetPolygonHoleCount( const FPolygonID PolygonID ) const;
-
-	/** Returns the number of vertices on the indexed hole contour for this polygon */
-	UFUNCTION( BlueprintPure, Category="Editable Mesh" )
-	int32 GetPolygonHoleVertexCount( const FPolygonID PolygonID, const int32 HoleNumber ) const;
-
-	/** Returns the indexed vertex on the given hole contour of the polygon */
-	UFUNCTION( BlueprintPure, Category="Editable Mesh" )
-	FVertexID GetPolygonHoleVertex( const FPolygonID PolygonID, const int32 HoleNumber, const int32 PolygonVertexNumber ) const;
-
-	/** Returns the indexed vertex instance on the given hole contour of the polygon */
-	UFUNCTION( BlueprintPure, Category="Editable Mesh" )
-	FVertexInstanceID GetPolygonHoleVertexInstance( const FPolygonID PolygonID, const int32 HoleNumber, const int32 PolygonVertexNumber ) const;
-
 	/** Returns the number of triangles which make up this polygon */
 	UFUNCTION( BlueprintPure, Category="Editable Mesh" )
 	int32 GetPolygonTriangulatedTriangleCount( const FPolygonID PolygonID ) const;
@@ -252,7 +236,7 @@ protected:
 	void GetConnectedSoftEdges( const FVertexID VertexID, TArray<FEdgeID>& OutConnectedSoftEdges ) const;
 	void GetVertexConnectedPolygonsInSameSoftEdgedGroup( const FVertexID VertexInstanceID, const FPolygonID PolygonID, TArray<FPolygonID>& OutPolygonIDs ) const;
 	FVertexInstanceID GetVertexInstanceInPolygonForVertex( const FPolygonID PolygonID, const FVertexID VertexID ) const;
-	void SetPolygonContourVertexAttributes( FMeshPolygonContour& Contour, const FPolygonID PolygonID, const int32 HoleIndex, const TArray<FMeshElementAttributeList>& AttributeLists );
+	void SetPolygonContourVertexAttributes( FMeshPolygonContour& Contour, const FPolygonID PolygonID, const TArray<FMeshElementAttributeList>& AttributeLists );
 	void SplitVertexInstanceInPolygons( const FVertexInstanceID VertexInstanceID, const TArray<FPolygonID>& PolygonIDs );
 	void ReplaceVertexInstanceInPolygons( const FVertexInstanceID OldVertexInstanceID, const FVertexInstanceID NewVertexInstanceID, const TArray<FPolygonID>& PolygonIDs );
 	float GetPolygonCornerAngleForVertex( const FPolygonID PolygonID, const FVertexID VertexID ) const;
@@ -413,20 +397,13 @@ public:
 	UFUNCTION( BlueprintPure, Category="Editable Mesh" ) void GetEdgeLoopElements( const FEdgeID EdgeID, TArray<FEdgeID>& EdgeLoopIDs ) const;
 	UFUNCTION( BlueprintPure, Category="Editable Mesh" ) FEdgeID GetEdgeThatConnectsVertices( const FVertexID VertexID0, const FVertexID VertexID1 ) const;
 	UFUNCTION( BlueprintPure, Category="Editable Mesh" ) int32 GetPolygonPerimeterEdgeCount( const FPolygonID PolygonID ) const;
-	UFUNCTION( BlueprintPure, Category="Editable Mesh" ) int32 GetPolygonHoleEdgeCount( const FPolygonID PolygonID, const int32 HoleNumber ) const;
 	UFUNCTION( BlueprintPure, Category="Editable Mesh" ) void GetPolygonPerimeterVertices( const FPolygonID PolygonID, TArray<FVertexID>& OutPolygonPerimeterVertexIDs ) const;
 	UFUNCTION( BlueprintPure, Category="Editable Mesh" ) void GetPolygonPerimeterVertexInstances( const FPolygonID PolygonID, TArray<FVertexInstanceID>& OutPolygonPerimeterVertexInstanceIDs ) const;
-	UFUNCTION( BlueprintPure, Category="Editable Mesh" ) void GetPolygonHoleVertices( const FPolygonID PolygonID, const int32 HoleNumber, TArray<FVertexID>& OutHoleVertexIDs ) const;
-	UFUNCTION( BlueprintPure, Category="Editable Mesh" ) void GetPolygonHoleVertexInstances( const FPolygonID PolygonID, const int32 HoleNumber, TArray<FVertexInstanceID>& OutHoleVertexInstanceIDs ) const;
 	UFUNCTION( BlueprintPure, Category="Editable Mesh" ) FEdgeID GetPolygonPerimeterEdge( const FPolygonID PolygonID, const int32 PerimeterEdgeNumber, bool& bOutEdgeWindingIsReversedForPolygon ) const;
-	UFUNCTION( BlueprintPure, Category="Editable Mesh" ) FEdgeID GetPolygonHoleEdge( const FPolygonID PolygonID, const int32 HoleNumber, const int32 HoleEdgeNumber ) const;
 	UFUNCTION( BlueprintPure, Category="Editable Mesh" ) void GetPolygonPerimeterEdges( const FPolygonID PolygonID, TArray<FEdgeID>& OutPolygonPerimeterEdgeIDs ) const;
-	UFUNCTION( BlueprintPure, Category="Editable Mesh" ) void GetPolygonHoleEdges( const FPolygonID PolygonID, const int32 HoleNumber, TArray<FEdgeID>& OutHoleEdgeIDs ) const;
 	UFUNCTION( BlueprintPure, Category="Editable Mesh" ) void GetPolygonAdjacentPolygons( const FPolygonID PolygonID, TArray<FPolygonID>& OutAdjacentPolygons ) const;
 	UFUNCTION( BlueprintPure, Category="Editable Mesh" ) int32 FindPolygonPerimeterVertexNumberForVertex( const FPolygonID PolygonID, const FVertexID VertexID ) const;
-	UFUNCTION( BlueprintPure, Category="Editable Mesh" ) int32 FindPolygonHoleVertexNumberForVertex( const FPolygonID PolygonID, const int32 HoleNumber, const FVertexID VertexID ) const;
 	UFUNCTION( BlueprintPure, Category="Editable Mesh" ) int32 FindPolygonPerimeterEdgeNumberForVertices( const FPolygonID PolygonID, const FVertexID EdgeVertexID0, const FVertexID EdgeVertexID1 ) const;
-	UFUNCTION( BlueprintPure, Category="Editable Mesh" ) int32 FindPolygonHoleEdgeNumberForVertices( const FPolygonID PolygonID, const int32 HoleNumber, const FVertexID EdgeVertexID0, const FVertexID EdgeVertexID1 ) const;
 	UFUNCTION( BlueprintPure, Category="Editable Mesh" ) FBox ComputeBoundingBox() const;
 	UFUNCTION( BlueprintPure, Category="Editable Mesh" ) FBoxSphereBounds ComputeBoundingBoxAndSphere() const;
 	UFUNCTION( BlueprintPure, Category="Editable Mesh" ) FVector ComputePolygonCenter( const FPolygonID PolygonID ) const;
@@ -445,7 +422,6 @@ public:
 	UFUNCTION( BlueprintCallable, Category="Editable Mesh" ) void SetSubdivisionCount( const int32 NewSubdivisionCount );
 	UFUNCTION( BlueprintCallable, Category="Editable Mesh" ) void MoveVertices( const TArray<FVertexToMove>& VerticesToMove );
 	UFUNCTION( BlueprintCallable, Category="Editable Mesh" ) void CreateMissingPolygonPerimeterEdges( const FPolygonID PolygonID, TArray<FEdgeID>& OutNewEdgeIDs );
-	UFUNCTION( BlueprintCallable, Category="Editable Mesh" ) void CreateMissingPolygonHoleEdges( const FPolygonID PolygonID, const int32 HoleNumber, TArray<FEdgeID>& OutNewEdgeIDs );
 	UFUNCTION( BlueprintCallable, Category="Editable Mesh" ) void SplitEdge( const FEdgeID EdgeID, const TArray<float>& Splits, TArray<FVertexID>& OutNewVertexIDs );
 	UFUNCTION( BlueprintCallable, Category="Editable Mesh" ) void InsertEdgeLoop( const FEdgeID EdgeID, const TArray<float>& Splits, TArray<FEdgeID>& OutNewEdgeIDs );
 	UFUNCTION( BlueprintCallable, Category="Editable Mesh" ) void SplitPolygons( const TArray<FPolygonToSplit>& PolygonsToSplit, TArray<FEdgeID>& OutNewEdgeIDs );
diff --git a/Engine/Source/Developer/MeshMergeUtilities/Private/MeshMergeHelpers.cpp b/Engine/Source/Developer/MeshMergeUtilities/Private/MeshMergeHelpers.cpp
--- a/Engine/Source/Developer/MeshMergeUtilities/Private/MeshMergeHelpers.cpp
+++ b/Engine/Source/Developer/MeshMergeUtilities/Private/MeshMergeHelpers.cpp
@@ -1315,19 +1315,8 @@ void FMeshMergeHelpers::AppendRawMesh(FMeshDescription& InTarget, const FMeshDes
 		TArray<FMeshDescription::FContourPoint> Contours;
 		CreateContour(SourceVertexInstanceIDs, Contours);
 
-		int32 HoleNumber = InSource.GetNumPolygonHoles(SourcePolygonID);
-		TArray<TArray<FMeshDescription::FContourPoint>> AllHoles;
-		AllHoles.AddDefaulted(HoleNumber);
-		for (int32 HoleIndex = 0; HoleIndex < HoleNumber; ++HoleIndex)
-		{
-			const FMeshPolygonContour& HoleContour = SourcePolygon.HoleContours[HoleIndex];
-			
-			TArray<FMeshDescription::FContourPoint>& Holes = AllHoles[HoleIndex];
-			CreateContour(HoleContour.VertexInstanceIDs, Holes);
-		}
-
 		// Insert a polygon into the mesh
-		const FPolygonID TargetPolygonID = InTarget.CreatePolygon(SourcePolygon.PolygonGroupID, Contours, AllHoles);
+		const FPolygonID TargetPolygonID = InTarget.CreatePolygon(SourcePolygon.PolygonGroupID, Contours);
 		//Triangulate the polygon
 		FMeshPolygon& Polygon = InTarget.GetPolygon(TargetPolygonID);
 		InTarget.ComputePolygonTriangulation(TargetPolygonID, Polygon.Triangles);
diff --git a/Engine/Source/Runtime/Core/Public/UObject/EditorObjectVersion.h b/Engine/Source/Runtime/Core/Public/UObject/EditorObjectVersion.h
--- a/Engine/Source/Runtime/Core/Public/UObject/EditorObjectVersion.h
+++ b/Engine/Source/Runtime/Core/Public/UObject/EditorObjectVersion.h
@@ -69,6 +69,8 @@ struct CORE_API FEditorObjectVersion
 		StaticMeshDeprecatedRawMesh,
 		// MeshDescriptionBulkData contains a Guid used as a DDC key
 		MeshDescriptionBulkDataGuid,
+		// Change to MeshDescription serialization (removed FMeshPolygon::HoleContours)
+		MeshDescriptionRemovedHoles,
 		// -----<new versions can be added above this line>-------------------------------------------------
 		VersionPlusOne,
 		LatestVersion = VersionPlusOne - 1
diff --git a/Engine/Source/Runtime/Engine/Private/StaticMesh.cpp b/Engine/Source/Runtime/Engine/Private/StaticMesh.cpp
--- a/Engine/Source/Runtime/Engine/Private/StaticMesh.cpp
+++ b/Engine/Source/Runtime/Engine/Private/StaticMesh.cpp
@@ -3025,7 +3025,7 @@ void UStaticMesh::FixupMaterialSlotName()
 // differences, etc.) replace the version GUID below with a new one.
 // In case of merge conflicts with DDC versions, you MUST generate a new GUID
 // and set this new GUID as the version.                                       
-#define MESHDATAKEY_STATICMESH_DERIVEDDATA_VER TEXT("0386934AE700020EBBC4DD45C0AA3C49")
+#define MESHDATAKEY_STATICMESH_DERIVEDDATA_VER TEXT("F62CF84A92D948DEB06D56AEC759D2D3")
 
 static const FString& GetMeshDataKeyStaticMeshDerivedDataVersion()
 {
diff --git a/Engine/Source/Runtime/MeshDescription/Private/MeshDescription.cpp b/Engine/Source/Runtime/MeshDescription/Private/MeshDescription.cpp
--- a/Engine/Source/Runtime/MeshDescription/Private/MeshDescription.cpp
+++ b/Engine/Source/Runtime/MeshDescription/Private/MeshDescription.cpp
@@ -79,11 +79,6 @@ FArchive& operator<<( FArchive& Ar, FMeshDescription& MeshDescription )
 
 			PopulatePolygonIDs( MeshDescription.GetPolygonPerimeterVertexInstances( PolygonID ) );
 
-			for( int32 HoleIndex = 0; HoleIndex < MeshDescription.GetNumPolygonHoles( PolygonID ); ++HoleIndex )
-			{
-				PopulatePolygonIDs( MeshDescription.GetPolygonHoleVertexInstances( PolygonID, HoleIndex ) );
-			}
-
 			const FPolygonGroupID PolygonGroupID = MeshDescription.PolygonArray[ PolygonID ].PolygonGroupID;
 			MeshDescription.PolygonGroupArray[ PolygonGroupID ].Polygons.Add( PolygonID );
 
@@ -208,14 +203,6 @@ void FMeshDescription::FixUpElementIDs( const FElementIDRemappings& Remappings )
 			VertexInstanceID = Remappings.GetRemappedVertexInstanceID( VertexInstanceID );
 		}
 
-		for( FMeshPolygonContour& HoleContour : Polygon.HoleContours )
-		{
-			for( FVertexInstanceID& VertexInstanceID : HoleContour.VertexInstanceIDs )
-			{
-				VertexInstanceID = Remappings.GetRemappedVertexInstanceID( VertexInstanceID );
-			}
-		}
-
 		for( FMeshTriangle& Triangle : Polygon.Triangles )
 		{
 			for( int32 TriangleVertexNumber = 0; TriangleVertexNumber < 3; ++TriangleVertexNumber )
@@ -281,22 +268,6 @@ void FMeshDescription::GetPolygonPerimeterVertices( const FPolygonID PolygonID,
 }
 
 
-void FMeshDescription::GetPolygonHoleVertices( const FPolygonID PolygonID, const int32 HoleIndex, TArray<FVertexID>& OutPolygonHoleVertexIDs ) const
-{
-	const FMeshPolygon& Polygon = GetPolygon( PolygonID );
-
-	OutPolygonHoleVertexIDs.SetNumUninitialized( Polygon.HoleContours[ HoleIndex ].VertexInstanceIDs.Num(), false );
-
-	int32 Index = 0;
-	for( const FVertexInstanceID VertexInstanceID : Polygon.HoleContours[ HoleIndex ].VertexInstanceIDs )
-	{
-		const FMeshVertexInstance& VertexInstance = VertexInstanceArray[ VertexInstanceID ];
-		OutPolygonHoleVertexIDs[ Index ] = VertexInstance.VertexID;
-		Index++;
-	}
-}
-
-
 /** Given three direction vectors, indicates if A and B are on the same 'side' of Vec. */
 bool FMeshDescription::VectorsOnSameSide( const FVector& Vec, const FVector& A, const FVector& B, const float SameSideDotProductEpsilon )
 {
@@ -399,7 +370,6 @@ void FMeshDescription::ComputePolygonTriangulation(const FPolygonID PolygonID, T
 
 	OutTriangles.Reset();
 
-	// @todo mesheditor holes: Does not support triangles with holes yet!
 	// @todo mesheditor: Perhaps should always attempt to triangulate by splitting polygons along the shortest edge, for better determinism.
 
 	//	const FMeshPolygon& Polygon = GetPolygon( PolygonID );
@@ -801,19 +771,6 @@ float FMeshDescription::GetPolygonCornerAngleForVertex(const FPolygonID PolygonI
 		// Return the internal angle if found
 		return GetContourAngle(Polygon.PerimeterContour, ContourIndex);
 	}
-	else
-	{
-		// If not found, look in all the holes
-		for (const FMeshPolygonContour& HoleContour : Polygon.HoleContours)
-		{
-			ContourIndex = HoleContour.VertexInstanceIDs.IndexOfByPredicate(IsVertexInstancedFromThisVertex);
-			if (ContourIndex != INDEX_NONE)
-			{
-				// Hole vertex contribution is the part which ISN'T the internal angle of the contour, so subtract from 2*pi
-				return (2.0f * PI) - GetContourAngle(HoleContour, ContourIndex);
-			}
-		}
-	}
 
 	// Found nothing; return 0
 	return 0.0f;
diff --git a/Engine/Source/Runtime/MeshDescription/Public/MeshDescription.h b/Engine/Source/Runtime/MeshDescription/Public/MeshDescription.h
--- a/Engine/Source/Runtime/MeshDescription/Public/MeshDescription.h
+++ b/Engine/Source/Runtime/MeshDescription/Public/MeshDescription.h
@@ -173,12 +173,6 @@ struct FMeshPolygon
 	/** The outer boundary edges of this polygon */
 	FMeshPolygonContour PerimeterContour;
 
-	/** Optional inner contours of this polygon that define holes inside of the polygon.  For the geometry to
-	    be considered valid, the hole contours should reside within the boundary of the polygon perimeter contour, 
-		and must not overlap each other.  No "nesting" of polygons inside the holes is supported -- those are 
-		simply separate polygons */
-	TArray<FMeshPolygonContour> HoleContours;
-
 	/** List of triangles which make up this polygon */
 	TArray<FMeshTriangle> Triangles;
 
@@ -189,7 +183,11 @@ struct FMeshPolygon
 	friend FArchive& operator<<( FArchive& Ar, FMeshPolygon& Polygon )
 	{
 		Ar << Polygon.PerimeterContour;
-		Ar << Polygon.HoleContours;
+		if (Ar.IsLoading() && Ar.CustomVer(FReleaseObjectVersion::GUID) < FEditorObjectVersion::MeshDescriptionRemovedHoles)
+		{
+			TArray<FMeshPolygonContour> Empty;
+			Ar << Empty;
+		}
 		if( Ar.IsLoading() && Ar.CustomVer( FReleaseObjectVersion::GUID ) < FReleaseObjectVersion::MeshDescriptionNewSerialization )
 		{
 			Ar << Polygon.Triangles;
@@ -527,15 +525,10 @@ private:
 		}
 	}
 
-	void CreatePolygon_Internal( const FPolygonID PolygonID, const FPolygonGroupID PolygonGroupID, const TArray<FContourPoint>& Perimeter, const TArray<TArray<FContourPoint>>& Holes )
+	void CreatePolygon_Internal( const FPolygonID PolygonID, const FPolygonGroupID PolygonGroupID, const TArray<FContourPoint>& Perimeter )
 	{
 		FMeshPolygon& Polygon = PolygonArray[ PolygonID ];
 		CreatePolygonContour_Internal( PolygonID, Perimeter, Polygon.PerimeterContour );
-		for( const TArray<FContourPoint>& Hole : Holes )
-		{
-			Polygon.HoleContours.Emplace();
-			CreatePolygonContour_Internal( PolygonID, Hole, Polygon.HoleContours.Last() );
-		}
 
 		Polygon.PolygonGroupID = PolygonGroupID;
 		PolygonGroupArray[ PolygonGroupID ].Polygons.Add( PolygonID );
@@ -545,18 +538,18 @@ private:
 
 public:
 	/** Adds a new polygon to the mesh and returns its ID */
-	FPolygonID CreatePolygon( const FPolygonGroupID PolygonGroupID, const TArray<FContourPoint>& Perimeter, const TArray<TArray<FContourPoint>>& Holes = TArray<TArray<FContourPoint>>() )
+	FPolygonID CreatePolygon( const FPolygonGroupID PolygonGroupID, const TArray<FContourPoint>& Perimeter )
 	{
 		const FPolygonID PolygonID = PolygonArray.Add();
-		CreatePolygon_Internal( PolygonID, PolygonGroupID, Perimeter, Holes );
+		CreatePolygon_Internal( PolygonID, PolygonGroupID, Perimeter );
 		return PolygonID;
 	}
 
 	/** Adds a new polygon to the mesh with the given ID */
-	void CreatePolygonWithID( const FPolygonID PolygonID, const FPolygonGroupID PolygonGroupID, const TArray<FContourPoint>& Perimeter, const TArray<TArray<FContourPoint>>& Holes = TArray<TArray<FContourPoint>>() )
+	void CreatePolygonWithID( const FPolygonID PolygonID, const FPolygonGroupID PolygonGroupID, const TArray<FContourPoint>& Perimeter )
 	{
 		PolygonArray.Insert( PolygonID );
-		CreatePolygon_Internal( PolygonID, PolygonGroupID, Perimeter, Holes );
+		CreatePolygon_Internal( PolygonID, PolygonGroupID, Perimeter );
 	}
 
 private:
@@ -595,11 +588,6 @@ public:
 		FMeshPolygon& Polygon = PolygonArray[ PolygonID ];
 		DeletePolygonContour_Internal( PolygonID, Polygon.PerimeterContour.VertexInstanceIDs, InOutOrphanedEdgesPtr, InOutOrphanedVertexInstancesPtr );
 
-		for( FMeshPolygonContour& HoleContour : Polygon.HoleContours )
-		{
-			DeletePolygonContour_Internal( PolygonID, HoleContour.VertexInstanceIDs, InOutOrphanedEdgesPtr, InOutOrphanedVertexInstancesPtr );
-		}
-
 		FMeshPolygonGroup& PolygonGroup = PolygonGroupArray[ Polygon.PolygonGroupID ];
 		verify( PolygonGroup.Polygons.Remove( PolygonID ) == 1 );
 
@@ -761,21 +749,6 @@ public:
 	/** Populates the passed array of VertexIDs with the vertices which form the polygon perimeter */
 	void GetPolygonPerimeterVertices( const FPolygonID PolygonID, TArray<FVertexID>& OutPolygonPerimeterVertexIDs ) const;
 
-	/** Returns the number of holes in this polygon */
-	int32 GetNumPolygonHoles( const FPolygonID PolygonID ) const
-	{
-		return PolygonArray[ PolygonID ].HoleContours.Num();
-	}
-
-	/** Populates the passed array of VertexIDs with the vertices which form the contour of the specified hole */
-	void GetPolygonHoleVertices( const FPolygonID PolygonID, const int32 HoleIndex, TArray<FVertexID>& OutPolygonHoleVertexIDs ) const;
-
-	/** Returns reference to an array of VertexInstance IDs forming the contour of the given hole */
-	const TArray<FVertexInstanceID>& GetPolygonHoleVertexInstances( const FPolygonID PolygonID, const int32 HoleIndex ) const
-	{
-		return PolygonArray[ PolygonID ].HoleContours[ HoleIndex ].VertexInstanceIDs;
-	}
-
 	void GetPolygonEdges(const FPolygonID PolygonID, TArray<FEdgeID>& OutPolygonEdges) const
 	{
 		const FMeshPolygonContour& PerimeterContour = PolygonArray[PolygonID].PerimeterContour;
