Change to FMeshDescription API. CreatePolygon[WithID] now takes a simple array of vertex instances instead of FMeshDescription::ContourPoint, hence there is no need to specify corresponding edge IDs. Any edges which do not exist in the mesh description will be automatically created by CreatePolygon, and their IDs returned.
Fixed up all code to use the new method, fixed a few bugs.

diff --git a/Engine/Plugins/Editor/MeshEditor/Source/PolygonModeling/SplitMesh.cpp b/Engine/Plugins/Editor/MeshEditor/Source/PolygonModeling/SplitMesh.cpp
--- a/Engine/Plugins/Editor/MeshEditor/Source/PolygonModeling/SplitMesh.cpp
+++ b/Engine/Plugins/Editor/MeshEditor/Source/PolygonModeling/SplitMesh.cpp
@@ -220,58 +220,34 @@ void USplitMeshCommand::Execute(IMeshEditorModeEditingContract& MeshEditorMode)
 		Mesh->GetMeshDescription()->VertexAttributes().SetAttribute<FVector>(NewVert, MeshAttribute::Vertex::Position, 0, Center);
 		FVertexInstanceID NewVertInstance = Mesh->GetMeshDescription()->CreateVertexInstance(NewVert);
 
-		TMap<FVertexID, FEdgeID> VertToEdgeIdMap;
 		FPolygonGroupID GroupId = Mesh->GetMeshDescription()->CreatePolygonGroup();
 		for(const FEdgeID EdgeId : BoundaryIds)
 		{
-			FEdgeID Edge0, Edge1;
-
 			FVertexID Vertex0 = Mesh->GetMeshDescription()->GetEdgeVertex(EdgeId, 0);
-			TArray<FVertexInstanceID> VertexInstances0 = Mesh->GetMeshDescription()->GetVertexVertexInstances(Vertex0);
+			const TArray<FVertexInstanceID>& VertexInstances0 = Mesh->GetMeshDescription()->GetVertexVertexInstances(Vertex0);
 			check(VertexInstances0.Num() > 0);
-			if(VertToEdgeIdMap.Contains(Vertex0))
-			{
-				Edge0 = VertToEdgeIdMap[Vertex0];
-			}
-			else
-			{
-				Edge0 = Mesh->GetMeshDescription()->CreateEdge(Vertex0, NewVert);
-				VertToEdgeIdMap.Add(Vertex0, Edge0);
-			}
 
 			FVertexID Vertex1 = Mesh->GetMeshDescription()->GetEdgeVertex(EdgeId, 1);
-			TArray<FVertexInstanceID> VertexInstances1 = Mesh->GetMeshDescription()->GetVertexVertexInstances(Vertex1);
+			const TArray<FVertexInstanceID>& VertexInstances1 = Mesh->GetMeshDescription()->GetVertexVertexInstances(Vertex1);
 			check(VertexInstances1.Num() > 0);
-			if(VertToEdgeIdMap.Contains(Vertex1))
-			{
-				Edge1 = VertToEdgeIdMap[Vertex1];
-			}
-			else
-			{
-				Edge1 = Mesh->GetMeshDescription()->CreateEdge(Vertex1, NewVert);
-				VertToEdgeIdMap.Add(Vertex1, Edge1);
-			}
-			TArray<FMeshDescription::FContourPoint> Edges;
-			Edges.SetNum(3);
+
+			TArray<FVertexInstanceID> PolygonVertexInstances;
+			PolygonVertexInstances.SetNum(3);
 			if(FVector::DotProduct(TransformedPlaneNormal, FVector::CrossProduct(VertexPositions[Vertex1] - VertexPositions[Vertex0], Center - VertexPositions[Vertex1])) < 0)
 			{
-				Edges[0].VertexInstanceID = VertexInstances0[0];
-				Edges[0].EdgeID = EdgeId;
-				Edges[1].VertexInstanceID = VertexInstances1[0];
-				Edges[1].EdgeID = Edge1;
-				Edges[2].VertexInstanceID = NewVertInstance;
-				Edges[2].EdgeID = Edge0;
+				PolygonVertexInstances[0] = VertexInstances0[0];
+				PolygonVertexInstances[1] = VertexInstances1[0];
+				PolygonVertexInstances[2] = NewVertInstance;
 			}
 			else
 			{
-				Edges[0].VertexInstanceID = VertexInstances1[0];
-				Edges[0].EdgeID = EdgeId;
-				Edges[1].VertexInstanceID = VertexInstances0[0];
-				Edges[1].EdgeID = Edge0;
-				Edges[2].VertexInstanceID = NewVertInstance;
-				Edges[2].EdgeID = Edge1;
+				PolygonVertexInstances[0] = VertexInstances1[0];
+				PolygonVertexInstances[1] = VertexInstances0[0];
+				PolygonVertexInstances[2] = NewVertInstance;
 			}
-			const auto NewPolygonID = Mesh->GetMeshDescription()->CreatePolygon(GroupId, Edges);
+			TArray<FEdgeID> NewEdgeIDs;
+			const auto NewPolygonID = Mesh->GetMeshDescription()->CreatePolygon(GroupId, PolygonVertexInstances, &NewEdgeIDs);
+			check(NewEdgeIDs.Num() == 0);
 			NewPolygonIds.Add(NewPolygonID);
 		}
 
@@ -281,9 +257,9 @@ void USplitMeshCommand::Execute(IMeshEditorModeEditingContract& MeshEditorMode)
 		NewPackageName = UPackageTools::SanitizePackageName(NewPackageName);
 		UPackage* NewPackage = CreatePackage(nullptr, *NewPackageName);
 		UStaticMesh* NewStaticMesh = NewObject<UStaticMesh>(NewPackage, *NewMeshName, RF_Public);
-		FMeshDescription* NewMeshDescription = NewStaticMesh->GetMeshDescription(0);
-		check(NewMeshDescription);
 		new (NewStaticMesh->SourceModels) FStaticMeshSourceModel();
+		FMeshDescription* NewMeshDescription = NewStaticMesh->CreateMeshDescription(0);
+		check(NewMeshDescription);
 
 		// @todo (mlentine): Need to make sure all numbers are the same
 		{
@@ -338,19 +314,8 @@ void USplitMeshCommand::Execute(IMeshEditorModeEditingContract& MeshEditorMode)
 				CopyAllAttributes(NewMeshDescription->PolygonGroupAttributes(), Mesh->GetMeshDescription()->PolygonGroupAttributes(), PolygonGroupId);
 				PolygonGroupSet.Add(PolygonGroupId);
 			}
-			TArray<FMeshDescription::FContourPoint> ContourPoints;
-			check(EdgeIds.Num() == VertexInstanceIds.Num());
-			
-			const int32 NumEdges = EdgeIds.Num();
-			for(int32 EdgeIdIndex = 0; EdgeIdIndex < NumEdges; ++EdgeIdIndex)
-			{
-				FMeshDescription::FContourPoint ContourPoint;
-				ContourPoint.VertexInstanceID = /*VertexInstanceIDs[j]*/ NewVertexInstanceIds[EdgeIdIndex];
-				ContourPoint.EdgeID = EdgeIds[EdgeIdIndex];
-				check(EdgeIds[EdgeIdIndex] == Mesh->GetMeshDescription()->GetVertexPairEdge(Mesh->GetMeshDescription()->GetVertexInstanceVertex(VertexInstanceIds[EdgeIdIndex]), Mesh->GetMeshDescription()->GetVertexInstanceVertex(VertexInstanceIds[(EdgeIdIndex + 1) % EdgeIds.Num()])));
-				ContourPoints.Add(ContourPoint);
-			}
-			NewMeshDescription->CreatePolygonWithID(PolygonId, PolygonGroupId, ContourPoints);
+
+			NewMeshDescription->CreatePolygonWithID(PolygonId, PolygonGroupId, NewVertexInstanceIds);
 			CopyAllAttributes(NewMeshDescription->PolygonAttributes(), Mesh->GetMeshDescription()->PolygonAttributes(), PolygonId);
 		}
 
diff --git a/Engine/Plugins/Editor/SpeedTreeImporter/Source/SpeedTreeImporter/Private/SpeedTreeImportFactory.cpp b/Engine/Plugins/Editor/SpeedTreeImporter/Source/SpeedTreeImporter/Private/SpeedTreeImportFactory.cpp
--- a/Engine/Plugins/Editor/SpeedTreeImporter/Source/SpeedTreeImporter/Private/SpeedTreeImportFactory.cpp
+++ b/Engine/Plugins/Editor/SpeedTreeImporter/Source/SpeedTreeImporter/Private/SpeedTreeImportFactory.cpp
@@ -1781,7 +1781,8 @@ UObject* USpeedTreeImportFactory::FactoryCreateBinary7(UClass* InClass, UObject*
 							//Create all vertex instance
 							for (int32 TriangleIndex = 0; TriangleIndex < TriangleCount; ++TriangleIndex)
 							{
-								FVertexInstanceID CornerVertexInstanceIDs[3];
+								TArray<FVertexInstanceID> CornerVertexInstanceIDs;
+								CornerVertexInstanceIDs.SetNum(3);
 								FVertexID CornerVertexIDs[3];
 								for (int32 Corner = 0; Corner < 3; ++Corner)
 								{
@@ -1792,33 +1793,8 @@ UObject* USpeedTreeImportFactory::FactoryCreateBinary7(UClass* InClass, UObject*
 									CornerVertexIDs[Corner] = MeshDescription->GetVertexInstanceVertex(CornerVertexInstanceIDs[Corner]);
 								}
 
-								TArray<FMeshDescription::FContourPoint> Contours;
-								for (int32 Corner = 0; Corner < 3; ++Corner)
-								{
-									int32 ContourPointIndex = Contours.AddDefaulted();
-									FMeshDescription::FContourPoint& ContourPoint = Contours[ContourPointIndex];
-									//Find the matching edge ID
-									uint32 CornerIndices[2];
-									CornerIndices[0] = (Corner + 0) % 3;
-									CornerIndices[1] = (Corner + 1) % 3;
-
-									FVertexID EdgeVertexIDs[2];
-									EdgeVertexIDs[0] = CornerVertexIDs[CornerIndices[0]];
-									EdgeVertexIDs[1] = CornerVertexIDs[CornerIndices[1]];
-
-									FEdgeID MatchEdgeId = MeshDescription->GetVertexPairEdge(EdgeVertexIDs[0], EdgeVertexIDs[1]);
-									if (MatchEdgeId == FEdgeID::Invalid)
-									{
-										MatchEdgeId = MeshDescription->CreateEdge(EdgeVertexIDs[0], EdgeVertexIDs[1]);
-										//All edge are smooth
-										EdgeHardnesses[MatchEdgeId] = false;
-										EdgeCreaseSharpnesses[MatchEdgeId] = 0.0f;
-									}
-									ContourPoint.EdgeID = MatchEdgeId;
-									ContourPoint.VertexInstanceID = CornerVertexInstanceIDs[CornerIndices[0]];
-								}
 								// Insert a polygon into the mesh
-								const FPolygonID NewPolygonID = MeshDescription->CreatePolygon(CurrentPolygonGroupID, Contours);
+								const FPolygonID NewPolygonID = MeshDescription->CreatePolygon(CurrentPolygonGroupID, CornerVertexInstanceIDs);
 								//Triangulate the polygon
 								FMeshPolygon& Polygon = MeshDescription->GetPolygon(NewPolygonID);
 								MeshDescription->ComputePolygonTriangulation(NewPolygonID, Polygon.Triangles);
@@ -1936,7 +1912,8 @@ UObject* USpeedTreeImportFactory::FactoryCreateBinary7(UClass* InClass, UObject*
 						int32 NumTriangles = NumIndices / 3;
 						for (int32 TriangleIndex = 0; TriangleIndex < NumTriangles; ++TriangleIndex)
 						{
-							FVertexInstanceID CornerVertexInstanceIDs[3];
+							TArray<FVertexInstanceID> CornerVertexInstanceIDs;
+							CornerVertexInstanceIDs.SetNum(3);
 							FVertexID CornerVertexIDs[3];
 							for (int32 Corner = 0; Corner < 3; ++Corner)
 							{
@@ -1965,33 +1942,8 @@ UObject* USpeedTreeImportFactory::FactoryCreateBinary7(UClass* InClass, UObject*
 								CornerVertexIDs[Corner] = VertexID;
 							}
 
-							TArray<FMeshDescription::FContourPoint> Contours;
-							for (int32 Corner = 0; Corner < 3; ++Corner)
-							{
-								int32 ContourPointIndex = Contours.AddDefaulted();
-								FMeshDescription::FContourPoint& ContourPoint = Contours[ContourPointIndex];
-								//Find the matching edge ID
-								uint32 CornerIndices[2];
-								CornerIndices[0] = (Corner + 0) % 3;
-								CornerIndices[1] = (Corner + 1) % 3;
-
-								FVertexID EdgeVertexIDs[2];
-								EdgeVertexIDs[0] = CornerVertexIDs[CornerIndices[0]];
-								EdgeVertexIDs[1] = CornerVertexIDs[CornerIndices[1]];
-
-								FEdgeID MatchEdgeId = MeshDescription->GetVertexPairEdge(EdgeVertexIDs[0], EdgeVertexIDs[1]);
-								if (MatchEdgeId == FEdgeID::Invalid)
-								{
-									MatchEdgeId = MeshDescription->CreateEdge(EdgeVertexIDs[0], EdgeVertexIDs[1]);
-									//All edge are smooth
-									EdgeHardnesses[MatchEdgeId] = false;
-									EdgeCreaseSharpnesses[MatchEdgeId] = 0.0f;
-								}
-								ContourPoint.EdgeID = MatchEdgeId;
-								ContourPoint.VertexInstanceID = CornerVertexInstanceIDs[CornerIndices[0]];
-							}
 							// Insert a polygon into the mesh
-							const FPolygonID NewPolygonID = MeshDescription->CreatePolygon(CurrentPolygonGroupID, Contours);
+							const FPolygonID NewPolygonID = MeshDescription->CreatePolygon(CurrentPolygonGroupID, CornerVertexInstanceIDs);
 							//Triangulate the polygon
 							FMeshPolygon& Polygon = MeshDescription->GetPolygon(NewPolygonID);
 							MeshDescription->ComputePolygonTriangulation(NewPolygonID, Polygon.Triangles);
@@ -2275,7 +2227,8 @@ UObject* USpeedTreeImportFactory::FactoryCreateBinary8(UClass* InClass, UObject*
 				int32 TriangleCount = DrawCall.m_uiIndexCount / 3;
 				for (int32 TriangleIndex = 0; TriangleIndex < TriangleCount; ++TriangleIndex)
 				{
-					FVertexInstanceID CornerVertexInstanceIDs[3];
+					TArray<FVertexInstanceID> CornerVertexInstanceIDs;
+					CornerVertexInstanceIDs.SetNum(3);
 					FVertexID CornerVertexIDs[3];
 					for (int32 Corner = 0; Corner < 3; ++Corner)
 					{
@@ -2373,33 +2326,8 @@ UObject* USpeedTreeImportFactory::FactoryCreateBinary8(UClass* InClass, UObject*
 						CornerVertexIDs[Corner] = VertexID;
 					}
 
-					TArray<FMeshDescription::FContourPoint> Contours;
-					for (int32 Corner = 0; Corner < 3; ++Corner)
-					{
-						int32 ContourPointIndex = Contours.AddDefaulted();
-						FMeshDescription::FContourPoint& ContourPoint = Contours[ContourPointIndex];
-						//Find the matching edge ID
-						uint32 CornerIndices[2];
-						CornerIndices[0] = (Corner + 0) % 3;
-						CornerIndices[1] = (Corner + 1) % 3;
-
-						FVertexID EdgeVertexIDs[2];
-						EdgeVertexIDs[0] = CornerVertexIDs[CornerIndices[0]];
-						EdgeVertexIDs[1] = CornerVertexIDs[CornerIndices[1]];
-
-						FEdgeID MatchEdgeId = MeshDescription->GetVertexPairEdge(EdgeVertexIDs[0], EdgeVertexIDs[1]);
-						if (MatchEdgeId == FEdgeID::Invalid)
-						{
-							MatchEdgeId = MeshDescription->CreateEdge(EdgeVertexIDs[0], EdgeVertexIDs[1]);
-							//All edge are smooth
-							EdgeHardnesses[MatchEdgeId] = false;
-							EdgeCreaseSharpnesses[MatchEdgeId] = 0.0f;
-						}
-						ContourPoint.EdgeID = MatchEdgeId;
-						ContourPoint.VertexInstanceID = CornerVertexInstanceIDs[CornerIndices[0]];
-					}
 					// Insert a polygon into the mesh
-					const FPolygonID NewPolygonID = MeshDescription->CreatePolygon(CurrentPolygonGroupID, Contours);
+					const FPolygonID NewPolygonID = MeshDescription->CreatePolygon(CurrentPolygonGroupID, CornerVertexInstanceIDs);
 					//Triangulate the polygon
 					FMeshPolygon& Polygon = MeshDescription->GetPolygon(NewPolygonID);
 					MeshDescription->ComputePolygonTriangulation(NewPolygonID, Polygon.Triangles);
diff --git a/Engine/Plugins/Experimental/AlembicImporter/Source/AlembicLibrary/Private/AbcImporter.cpp b/Engine/Plugins/Experimental/AlembicImporter/Source/AlembicLibrary/Private/AbcImporter.cpp
--- a/Engine/Plugins/Experimental/AlembicImporter/Source/AlembicLibrary/Private/AbcImporter.cpp
+++ b/Engine/Plugins/Experimental/AlembicImporter/Source/AlembicLibrary/Private/AbcImporter.cpp
@@ -1308,7 +1308,8 @@ void FAbcImporter::GenerateMeshDescriptionFromSample(const FAbcMeshSample* Sampl
 	uint32 TriangleCount = Sample->Indices.Num() / 3;
 	for (uint32 TriangleIndex = 0; TriangleIndex < TriangleCount; ++TriangleIndex)
 	{
-		FVertexInstanceID CornerVertexInstanceIDs[3];
+		TArray<FVertexInstanceID> CornerVertexInstanceIDs;
+		CornerVertexInstanceIDs.SetNum(3);
 		FVertexID CornerVertexIDs[3];
 		for (int32 Corner = 0; Corner < 3; ++Corner)
 		{
@@ -1343,34 +1344,9 @@ void FAbcImporter::GenerateMeshDescriptionFromSample(const FAbcMeshSample* Sampl
 			CornerVertexIDs[Corner] = VertexID;
 		}
 
-		TArray<FMeshDescription::FContourPoint> Contours;
-		for (int32 Corner = 0; Corner < 3; ++Corner)
-		{
-			int32 ContourPointIndex = Contours.AddDefaulted();
-			FMeshDescription::FContourPoint& ContourPoint = Contours[ContourPointIndex];
-			//Find the matching edge ID
-			uint32 CornerIndices[2];
-			CornerIndices[0] = (Corner + 0) % 3;
-			CornerIndices[1] = (Corner + 1) % 3;
-
-			FVertexID EdgeVertexIDs[2];
-			EdgeVertexIDs[0] = CornerVertexIDs[CornerIndices[0]];
-			EdgeVertexIDs[1] = CornerVertexIDs[CornerIndices[1]];
-
-			FEdgeID MatchEdgeId = MeshDescription->GetVertexPairEdge(EdgeVertexIDs[0], EdgeVertexIDs[1]);
-			if (MatchEdgeId == FEdgeID::Invalid)
-			{
-				MatchEdgeId = MeshDescription->CreateEdge(EdgeVertexIDs[0], EdgeVertexIDs[1]);
-				//Conversion from smooth group to hard edges will be done at the end
-				EdgeHardnesses[MatchEdgeId] = false;
-				EdgeCreaseSharpnesses[MatchEdgeId] = 0.0f;
-			}
-			ContourPoint.EdgeID = MatchEdgeId;
-			ContourPoint.VertexInstanceID = CornerVertexInstanceIDs[CornerIndices[0]];
-		}
 		const FPolygonGroupID PolygonGroupID(Sample->MaterialIndices[TriangleIndex]);
 		// Insert a polygon into the mesh
-		const FPolygonID NewPolygonID = MeshDescription->CreatePolygon(PolygonGroupID, Contours);
+		const FPolygonID NewPolygonID = MeshDescription->CreatePolygon(PolygonGroupID, CornerVertexInstanceIDs);
 		//Triangulate the polygon
 		FMeshPolygon& Polygon = MeshDescription->GetPolygon(NewPolygonID);
 		MeshDescription->ComputePolygonTriangulation(NewPolygonID, Polygon.Triangles);
diff --git a/Engine/Plugins/Importers/USDImporter/Source/USDImporter/Private/StaticMeshImporter.cpp b/Engine/Plugins/Importers/USDImporter/Source/USDImporter/Private/StaticMeshImporter.cpp
--- a/Engine/Plugins/Importers/USDImporter/Source/USDImporter/Private/StaticMeshImporter.cpp
+++ b/Engine/Plugins/Importers/USDImporter/Source/USDImporter/Private/StaticMeshImporter.cpp
@@ -145,7 +145,6 @@ bool FUSDStaticMeshImportState::AddPolygons(FMeshDescriptionWrapper& DestMeshWra
 	TMap<int32, FPolygonGroupID> PolygonGroupMapping;
 	TArray<FVertexInstanceID> CornerInstanceIDs;
 	TArray<FVertexID> CornerVerticesIDs;
-	TArray<FMeshDescription::FContourPoint> Contours;
 	int32 CurrentVertexInstanceIndex = 0;
 
 	bool bFlipThisGeometry = bFlip;
@@ -241,33 +240,9 @@ bool FUSDStaticMeshImportState::AddPolygons(FMeshDescriptionWrapper& DestMeshWra
 			PolygonGroupMapping.Add(RealMaterialIndex, ExistingPolygonGroup);
 		}
 
-		// Create polygon edges
-		Contours.Reset();
-		for (uint32 PolygonEdgeNumber = 0; PolygonEdgeNumber < (uint32)PolygonVertexCount; ++PolygonEdgeNumber)
-		{
-			int32 ContourPointIndex = Contours.AddDefaulted();
-			FMeshDescription::FContourPoint& ContourPoint = Contours[ContourPointIndex];
-			//Find the matching edge ID
-			uint32 CornerIndices[2];
-			CornerIndices[0] = (PolygonEdgeNumber + 0) % PolygonVertexCount;
-			CornerIndices[1] = (PolygonEdgeNumber + 1) % PolygonVertexCount;
-
-			FVertexID EdgeVertexIDs[2];
-			EdgeVertexIDs[0] = CornerVerticesIDs[CornerIndices[0]];
-			EdgeVertexIDs[1] = CornerVerticesIDs[CornerIndices[1]];
-
-			FEdgeID MatchEdgeId = MeshDescription->GetVertexPairEdge(EdgeVertexIDs[0], EdgeVertexIDs[1]);
-			if (MatchEdgeId == FEdgeID::Invalid)
-			{
-				MatchEdgeId = MeshDescription->CreateEdge(EdgeVertexIDs[0], EdgeVertexIDs[1]);
-			}
-			ContourPoint.EdgeID = MatchEdgeId;
-			ContourPoint.VertexInstanceID = CornerInstanceIDs[CornerIndices[0]];
-		}
-
 		FPolygonGroupID PolygonGroupID = PolygonGroupMapping[RealMaterialIndex];
 		// Insert a polygon into the mesh
-		const FPolygonID NewPolygonID = MeshDescription->CreatePolygon(PolygonGroupID, Contours);
+		const FPolygonID NewPolygonID = MeshDescription->CreatePolygon(PolygonGroupID, CornerInstanceIDs);
 		if (bFlipThisGeometry)
 		{
 			MeshDescription->ReversePolygonFacing(NewPolygonID);
diff --git a/Engine/Plugins/Runtime/EditableMesh/Source/EditableMesh/EditableGeometryCollectionAdapter.cpp b/Engine/Plugins/Runtime/EditableMesh/Source/EditableMesh/EditableGeometryCollectionAdapter.cpp
--- a/Engine/Plugins/Runtime/EditableMesh/Source/EditableMesh/EditableGeometryCollectionAdapter.cpp
+++ b/Engine/Plugins/Runtime/EditableMesh/Source/EditableMesh/EditableGeometryCollectionAdapter.cpp
@@ -176,7 +176,9 @@ void UEditableGeometryCollectionAdapter::InitEditableGeometryCollection( UEditab
 						if ((GCBoneMap[Indices[0]] != RenderingSectionIndex) || GCVisible[TriangleIndex] == false)
 							continue;
 
-						FVertexInstanceID TriangleVertexInstanceIDs[3];
+						TArray<FVertexInstanceID> TriangleVertexInstanceIDs;
+						TriangleVertexInstanceIDs.SetNum(3);
+
 						FVertexID TriangleVertexIDs[3];
 						for (uint32 TriangleVertexIndex = 0; TriangleVertexIndex < 3; ++TriangleVertexIndex)
 						{
@@ -194,47 +196,20 @@ void UEditableGeometryCollectionAdapter::InitEditableGeometryCollection( UEditab
 
 						if (bIsValidTriangle)
 						{
-							FEdgeID NewEdgeIDs[3];
-
-							// Connect edges
-							{
-								// Add the edges of this triangle
-								for (uint32 TriangleEdgeNumber = 0; TriangleEdgeNumber < 3; ++TriangleEdgeNumber)
-								{
-									const FVertexID VertexID0 = TriangleVertexIDs[TriangleEdgeNumber];
-									const FVertexID VertexID1 = TriangleVertexIDs[(TriangleEdgeNumber + 1) % 3];
-
-									FEdgeID NewEdgeID = MeshDescription->GetVertexPairEdge(VertexID0, VertexID1);
-
-									if (NewEdgeID == FEdgeID::Invalid)
-									{
-										NewEdgeID = MeshDescription->CreateEdge(VertexID0, VertexID1);
-									}
-
-									NewEdgeIDs[TriangleEdgeNumber] = NewEdgeID;
-								}
-							}
-
 							// Geometry Collections only support triangles, so there's no need to triangulate anything yet.  We'll make both
 							// a triangle and a polygon here.
 							const FAdaptorTriangleID NewTriangleID = FAdaptorTriangleID(TriangleGroupIndex++);
 
 							NewRenderingPolygonGroup.Triangles.Insert(NewTriangleID);
 							FMeshTriangle& NewTriangle = NewRenderingPolygonGroup.Triangles[NewTriangleID];
-							static TArray<FMeshDescription::FContourPoint> Perimeter;
-							Perimeter.Reset(3);
-							Perimeter.AddUninitialized(3);
 							for (uint32 TriangleVertexIndex = 0; TriangleVertexIndex < 3; ++TriangleVertexIndex)
 							{
-								Perimeter[TriangleVertexIndex].VertexInstanceID = TriangleVertexInstanceIDs[TriangleVertexIndex];
-								Perimeter[TriangleVertexIndex].EdgeID = NewEdgeIDs[TriangleVertexIndex];
-
 								NewTriangle.SetVertexInstanceID(TriangleVertexIndex, TriangleVertexInstanceIDs[TriangleVertexIndex]);
 							}
 
 							// Insert a polygon into the mesh
 							const FPolygonID NewPolygonID(TriangleIndex);
-							MeshDescription->CreatePolygonWithID(NewPolygonID, NewPolygonGroupID, Perimeter);
+							MeshDescription->CreatePolygonWithID(NewPolygonID, NewPolygonGroupID, TriangleVertexInstanceIDs);
 
 							// Create a rendering polygon mirror, indexed by the same ID
 							RenderingPolygons.Insert(NewPolygonID);
diff --git a/Engine/Plugins/Runtime/EditableMesh/Source/EditableMesh/EditableMesh.cpp b/Engine/Plugins/Runtime/EditableMesh/Source/EditableMesh/EditableMesh.cpp
--- a/Engine/Plugins/Runtime/EditableMesh/Source/EditableMesh/EditableMesh.cpp
+++ b/Engine/Plugins/Runtime/EditableMesh/Source/EditableMesh/EditableMesh.cpp
@@ -273,12 +273,11 @@ template <typename ElementIDType>
 static void CopyAllAttributes( TAttributesSet<ElementIDType>& DestAttributesSet, const TAttributesSet<ElementIDType>& SrcAttributesSet, const ElementIDType ElementID )
 {
 	SrcAttributesSet.ForEachAttributeIndicesArray(
-		[ &DestAttributesSet, ElementID ]( const FName AttributeName, const auto& AttributeIndicesArray )
+		[ &DestAttributesSet, ElementID ]( const FName AttributeName, auto AttributeArrayRef )
 		{
-			for( int32 Index = 0; Index < AttributeIndicesArray.GetNumIndices(); ++Index )
+			for( int32 Index = 0; Index < AttributeArrayRef.GetNumIndices(); ++Index )
 			{
-				using AttributeType = decltype( AttributeIndicesArray.GetDefaultValue() );
-				DestAttributesSet.template SetAttribute<AttributeType>( ElementID, AttributeName, Index, AttributeIndicesArray.GetArrayForIndex( Index )[ ElementID ] );
+				DestAttributesSet.SetAttribute( ElementID, AttributeName, Index, AttributeArrayRef.Get( ElementID, Index ) );
 			}
 		}
 	);
@@ -3985,91 +3984,24 @@ FVertexInstanceID UEditableMesh::CreateVertexInstanceForContourVertex( const FVe
 }
 
 
-void UEditableMesh::CreatePolygonContour( const TArray<FVertexAndAttributes>& Contour, TArray<FEdgeID>& OutExistingEdgeIDs, TArray<FEdgeID>& OutNewEdgeIDs, TArray<FMeshDescription::FContourPoint>& OutContourPoints )
+void UEditableMesh::CreatePolygonContour( const TArray<FVertexAndAttributes>& Contour, TArray<FVertexInstanceID>& OutVertexInstanceIDs )
 {
 	// All polygons must have at least three vertices
 	const int32 NumContourVertices = Contour.Num();
 	check( NumContourVertices >= 3 );
 
-	OutExistingEdgeIDs.Reset( NumContourVertices );
-	OutNewEdgeIDs.Reset( NumContourVertices );
-	OutContourPoints.SetNumUninitialized( NumContourVertices );
-
-	static TArray<FEdgeToCreate> EdgesToCreate;
-	EdgesToCreate.Reset();
+	OutVertexInstanceIDs.SetNumUninitialized( NumContourVertices );
 
 	static TArray<FVertexInstanceToCreate> VertexInstancesToCreate;
 	VertexInstancesToCreate.Reset();
 
-	// Get edges or decide to create them
-	for( int32 VertexNumber = 0; VertexNumber < NumContourVertices; ++VertexNumber )
-	{
-		const int32 NextVertexNumber = ( VertexNumber + 1 ) % NumContourVertices;
-
-		// Get vertex ID for the ends of the edge.
-		FVertexID VertexID0 = Contour[ VertexNumber ].VertexID;
-		FVertexID VertexID1 = Contour[ NextVertexNumber ].VertexID;
-
-		// If valid vertex IDs weren't supplied, we'll deduce it from the vertex instance IDs instead (which must be valid instead).
-		if( VertexID0 == FVertexID::Invalid )
-		{
-			const FVertexInstanceID VertexInstanceID0 = Contour[ VertexNumber ].VertexInstanceID;
-			check( VertexInstanceID0 != FVertexInstanceID::Invalid );
-			VertexID0 = GetVertexInstanceVertex( VertexInstanceID0 );
-		}
-
-		if( VertexID1 == FVertexID::Invalid )
-		{
-			const FVertexInstanceID VertexInstanceID1 = Contour[ NextVertexNumber ].VertexInstanceID;
-			check( VertexInstanceID1 != FVertexInstanceID::Invalid );
-			VertexID1 = GetVertexInstanceVertex( VertexInstanceID1 );
-		}
-
-		FEdgeID EdgeID = GetMeshDescription()->GetVertexPairEdge( VertexID0, VertexID1 );
-
-		// Copy the edge ID to the output.
-		// If it is invalid, it will be patched later when new edges are created.
-		OutContourPoints[ VertexNumber ].EdgeID = EdgeID;
-
-		if( EdgeID == FEdgeID::Invalid )
-		{
-			// Create the new edge. Don't specify any attributes now.
-			EdgesToCreate.Emplace();
-			FEdgeToCreate& EdgeToCreate = EdgesToCreate.Last();
-			EdgeToCreate.VertexID0 = VertexID0;
-			EdgeToCreate.VertexID1 = VertexID1;
-		}
-		else
-		{
-			OutExistingEdgeIDs.Add( EdgeID );
-		}
-	}
-
-	// Create missing edges
-	if( EdgesToCreate.Num() > 0 )
-	{
-		CreateEdges( EdgesToCreate, OutNewEdgeIDs );
-
-		// Fill in missing edge indices
-		int32 NewEdgeIndex = 0;
-
-		for( FMeshDescription::FContourPoint& ContourPoint : OutContourPoints )
-		{
-			if( ContourPoint.EdgeID == FEdgeID::Invalid )
-			{
-				ContourPoint.EdgeID = OutNewEdgeIDs[ NewEdgeIndex ];
-				NewEdgeIndex++;
-			}
-		}
-	}
-
 	// Assign vertex instances to the polygon
 	for( int32 VertexNumber = 0; VertexNumber < NumContourVertices; ++VertexNumber )
 	{
 		// Copy the supplied vertex instance ID into the output.
 		// If it is valid, there is nothing more to be done.
 		// If it is not valid, copy the invalid ID. They will be replaced by the IDs of the new vertex instances in a later pass.
-		OutContourPoints[ VertexNumber ].VertexInstanceID = Contour[ VertexNumber ].VertexInstanceID;
+		OutVertexInstanceIDs[ VertexNumber ] = Contour[ VertexNumber ].VertexInstanceID;
 
 		if( Contour[ VertexNumber ].VertexInstanceID != FVertexInstanceID::Invalid )
 		{
@@ -4108,11 +4040,11 @@ void UEditableMesh::CreatePolygonContour( const TArray<FVertexAndAttributes>& Co
 		// Fill in missing vertex instance indices
 		int32 NewVertexInstanceIndex = 0;
 
-		for( FMeshDescription::FContourPoint& ContourPoint : OutContourPoints )
+		for( FVertexInstanceID& VertexInstanceID : OutVertexInstanceIDs )
 		{
-			if( ContourPoint.VertexInstanceID == FEdgeID::Invalid )
+			if( VertexInstanceID == FVertexInstanceID::Invalid )
 			{
-				ContourPoint.VertexInstanceID = NewVertexInstanceIDs[ NewVertexInstanceIndex ];
+				VertexInstanceID = NewVertexInstanceIDs[ NewVertexInstanceIndex ];
 				NewVertexInstanceIndex++;
 			}
 		}
@@ -4132,54 +4064,59 @@ void UEditableMesh::CreatePolygons( const TArray<FPolygonToCreate>& PolygonsToCr
 		static TArray<FAttributesForEdge> AttributesForEdges;
 		AttributesForEdges.Reset();
 
-		static TArray<FEdgeID> ExistingEdgeIDs;
-		ExistingEdgeIDs.Reset();
-
 		GetMeshDescription()->ReserveNewPolygons( PolygonsToCreate.Num() );
 		for( const FPolygonToCreate& PolygonToCreate : PolygonsToCreate )
 		{
-			static TArray<FMeshDescription::FContourPoint> PerimeterContourPoints;
+			static TArray<FVertexInstanceID> PerimeterVertexInstances;
 			static TArray<FEdgeID> NewEdgeIDsForContour;
-			static TArray<FEdgeID> ExistingEdgeIDsForContour;
 
-			CreatePolygonContour( PolygonToCreate.PerimeterVertices, ExistingEdgeIDsForContour, NewEdgeIDsForContour, PerimeterContourPoints );
-			ExistingEdgeIDs.Append( ExistingEdgeIDsForContour );
-			OutNewEdgeIDs.Append( NewEdgeIDsForContour );
+			CreatePolygonContour( PolygonToCreate.PerimeterVertices, PerimeterVertexInstances );
 
 			FPolygonID PolygonID = PolygonToCreate.OriginalPolygonID;
 			if( PolygonID != FPolygonID::Invalid )
 			{
-				GetMeshDescription()->CreatePolygonWithID( PolygonID, PolygonToCreate.PolygonGroupID, PerimeterContourPoints );
+				GetMeshDescription()->CreatePolygonWithID( PolygonID, PolygonToCreate.PolygonGroupID, PerimeterVertexInstances, &NewEdgeIDsForContour );
 			}
 			else
 			{
-				PolygonID = GetMeshDescription()->CreatePolygon( PolygonToCreate.PolygonGroupID, PerimeterContourPoints );
+				PolygonID = GetMeshDescription()->CreatePolygon( PolygonToCreate.PolygonGroupID, PerimeterVertexInstances, &NewEdgeIDsForContour );
 			}
 
 			OutNewPolygonIDs.Add( PolygonID );
+			OutNewEdgeIDs.Append( NewEdgeIDsForContour );
 
 			// Set edge hardness for polygon edges according to PolygonToCreate mode
-			AttributesForEdges.Reserve( AttributesForEdges.Num() + OutNewEdgeIDs.Num() + ExistingEdgeIDs.Num() );
+			const int32 ExistingEdgeCount = PerimeterVertexInstances.Num() - NewEdgeIDsForContour.Num();
+			AttributesForEdges.Reserve( AttributesForEdges.Num() + ExistingEdgeCount );
 
+			// Set directly hardness of any edges which were automatically created by the polygon.
+			// This needn't be done transactionally as we explicitly add the undo action later on.
 			const bool bHardEdge = ( PolygonToCreate.PolygonEdgeHardness == EPolygonEdgeHardness::AllEdgesHard || PolygonToCreate.PolygonEdgeHardness == EPolygonEdgeHardness::NewEdgesHard );
 			for( const FEdgeID EdgeID : OutNewEdgeIDs )
 			{
-				AttributesForEdges.Emplace();
-				FAttributesForEdge& AttributesForEdge = AttributesForEdges.Last();
-				AttributesForEdge.EdgeID = EdgeID;
-				AttributesForEdge.EdgeAttributes.Attributes.Emplace( MeshAttribute::Edge::IsHard, 0, FMeshElementAttributeValue( bHardEdge ) );
+				GetMeshDescription()->EdgeAttributes().SetAttribute( EdgeID, MeshAttribute::Edge::IsHard, 0, bHardEdge );
 			}
 
 			// If we are setting all edges' hardnesses (not just new ones), set existing edges' attributes here.
 			// This will also split any vertex instances which have just been included in a polygon if necessary.
-			if( PolygonToCreate.PolygonEdgeHardness == EPolygonEdgeHardness::AllEdgesHard || PolygonToCreate.PolygonEdgeHardness == EPolygonEdgeHardness::AllEdgesSoft )
+			if( ExistingEdgeCount > 0 &&
+				( PolygonToCreate.PolygonEdgeHardness == EPolygonEdgeHardness::AllEdgesHard || PolygonToCreate.PolygonEdgeHardness == EPolygonEdgeHardness::AllEdgesSoft ) )
 			{
-				for( const FEdgeID EdgeID : ExistingEdgeIDs )
+				FVertexInstanceID LastVertexInstanceID = PerimeterVertexInstances.Last();
+				for( const FVertexInstanceID VertexInstanceID : PerimeterVertexInstances )
 				{
-					AttributesForEdges.Emplace();
-					FAttributesForEdge& AttributesForEdge = AttributesForEdges.Last();
-					AttributesForEdge.EdgeID = EdgeID;
-					AttributesForEdge.EdgeAttributes.Attributes.Emplace( MeshAttribute::Edge::IsHard, 0, FMeshElementAttributeValue( bHardEdge ) );
+					const FVertexID VertexID0 = GetMeshDescription()->GetVertexInstanceVertex( VertexInstanceID );
+					const FVertexID VertexID1 = GetMeshDescription()->GetVertexInstanceVertex( LastVertexInstanceID );
+					const FEdgeID EdgeID = GetMeshDescription()->GetVertexPairEdge( VertexID0, VertexID1 );
+					check( EdgeID != FEdgeID::Invalid );
+					if( !NewEdgeIDsForContour.Contains( EdgeID ) )
+					{
+						AttributesForEdges.Emplace();
+						FAttributesForEdge& AttributesForEdge = AttributesForEdges.Last();
+						AttributesForEdge.EdgeID = EdgeID;
+						AttributesForEdge.EdgeAttributes.Attributes.Emplace( MeshAttribute::Edge::IsHard, 0, FMeshElementAttributeValue( bHardEdge ) );
+					}
+					LastVertexInstanceID = VertexInstanceID;
 				}
 			}
 		}
@@ -4204,6 +4141,20 @@ void UEditableMesh::CreatePolygons( const TArray<FPolygonToCreate>& PolygonsToCr
 	// Generate triangles for the new polygon
 	PolygonsPendingTriangulation.Append( OutNewPolygonIDs );
 
+	// If any new edges were automatically created when creating the polygons, add an action to the undo stack to delete them here
+	if( OutNewEdgeIDs.Num() > 0 )
+	{
+		FDeleteEdgesChangeInput RevertInput;
+		RevertInput.bDeleteOrphanedVertices = false;
+		RevertInput.EdgeIDsToDelete.Reserve( OutNewEdgeIDs.Num() );
+		for( int32 EdgeNumber = OutNewEdgeIDs.Num() - 1; EdgeNumber >= 0; --EdgeNumber )
+		{
+			RevertInput.EdgeIDsToDelete.Add( OutNewEdgeIDs[ EdgeNumber ] );
+		}
+
+		AddUndo( MakeUnique<FDeleteEdgesChange>( MoveTemp( RevertInput ) ) );
+	}
+
 	// NOTE: We iterate backwards, to delete polygons in the opposite order that we added them
 	{
 		FDeletePolygonsChangeInput RevertInput;
diff --git a/Engine/Plugins/Runtime/EditableMesh/Source/EditableMesh/EditableStaticMeshAdapter.cpp b/Engine/Plugins/Runtime/EditableMesh/Source/EditableMesh/EditableStaticMeshAdapter.cpp
--- a/Engine/Plugins/Runtime/EditableMesh/Source/EditableMesh/EditableStaticMeshAdapter.cpp
+++ b/Engine/Plugins/Runtime/EditableMesh/Source/EditableMesh/EditableStaticMeshAdapter.cpp
@@ -295,7 +295,9 @@ void UEditableStaticMeshAdapter::InitEditableStaticMesh( UEditableMesh* Editable
 					{
 						const uint32 RenderingTriangleFirstVertexIndex = FRenderingPolygonGroup::TriangleIndexToRenderingTriangleFirstIndex( RenderingSection, FTriangleID( SectionTriangleIndex ) );
 
-						FVertexInstanceID TriangleVertexInstanceIDs[ 3 ];
+						TArray<FVertexInstanceID> TriangleVertexInstanceIDs;
+						TriangleVertexInstanceIDs.SetNum( 3 );
+
 						FVertexID TriangleVertexIDs[ 3 ];
 						for( uint32 TriangleVertexIndex = 0; TriangleVertexIndex < 3; ++TriangleVertexIndex )
 						{
@@ -313,27 +315,6 @@ void UEditableStaticMeshAdapter::InitEditableStaticMesh( UEditableMesh* Editable
 
 						if( bIsValidTriangle )
 						{
-							FEdgeID NewEdgeIDs[ 3 ];
-
-							// Connect edges
-							{
-								// Add the edges of this triangle
-								for( uint32 TriangleEdgeNumber = 0; TriangleEdgeNumber < 3; ++TriangleEdgeNumber )
-								{
-									const FVertexID VertexID0 = TriangleVertexIDs[ TriangleEdgeNumber ];
-									const FVertexID VertexID1 = TriangleVertexIDs[ ( TriangleEdgeNumber + 1 ) % 3 ];
-									
-									FEdgeID NewEdgeID = MeshDescription->GetVertexPairEdge( VertexID0, VertexID1 );
-
-									if( NewEdgeID == FEdgeID::Invalid )
-									{
-										NewEdgeID = MeshDescription->CreateEdge( VertexID0, VertexID1 );
-									}
-
-									NewEdgeIDs[ TriangleEdgeNumber ] = NewEdgeID;
-								}
-							}
-
 							// Static meshes only support triangles, so there's no need to triangulate anything yet.  We'll make both
 							// a triangle and a polygon here.
 							const FTriangleID NewTriangleID = FTriangleID( SectionTriangleIndex );
@@ -341,19 +322,13 @@ void UEditableStaticMeshAdapter::InitEditableStaticMesh( UEditableMesh* Editable
 							NewRenderingPolygonGroup.Triangles.Insert( NewTriangleID );
 							FMeshTriangle& NewTriangle = NewRenderingPolygonGroup.Triangles[ NewTriangleID ];
 
-							static TArray<FMeshDescription::FContourPoint> Perimeter;
-							Perimeter.Reset( 3 );
-							Perimeter.AddUninitialized( 3 );
 							for( uint32 TriangleVertexIndex = 0; TriangleVertexIndex < 3; ++TriangleVertexIndex )
 							{
-								Perimeter[ TriangleVertexIndex ].VertexInstanceID = TriangleVertexInstanceIDs[ TriangleVertexIndex ];
-								Perimeter[ TriangleVertexIndex ].EdgeID = NewEdgeIDs[ TriangleVertexIndex ];
-
 								NewTriangle.SetVertexInstanceID( TriangleVertexIndex, TriangleVertexInstanceIDs[ TriangleVertexIndex ] );
 							}
 
 							// Insert a polygon into the mesh
-							const FPolygonID NewPolygonID = MeshDescription->CreatePolygon( NewPolygonGroupID, Perimeter );
+							const FPolygonID NewPolygonID = MeshDescription->CreatePolygon( NewPolygonGroupID, TriangleVertexInstanceIDs );
 
 							// Create a rendering polygon mirror, indexed by the same ID
 							RenderingPolygons.Insert( NewPolygonID );
diff --git a/Engine/Plugins/Runtime/EditableMesh/Source/EditableMesh/Public/EditableMesh.h b/Engine/Plugins/Runtime/EditableMesh/Source/EditableMesh/Public/EditableMesh.h
--- a/Engine/Plugins/Runtime/EditableMesh/Source/EditableMesh/Public/EditableMesh.h
+++ b/Engine/Plugins/Runtime/EditableMesh/Source/EditableMesh/Public/EditableMesh.h
@@ -231,7 +231,7 @@ protected:
 	void SetPolygonAttribute( const FPolygonID PolygonID, const FMeshElementAttributeData& Attribute );
 	void SetPolygonGroupAttribute( const FPolygonGroupID PolygonGroupID, const FMeshElementAttributeData& Attribute );
 	FVertexInstanceID CreateVertexInstanceForContourVertex( const FVertexAndAttributes& ContourVertex, const FPolygonID PolygonID );
-	void CreatePolygonContour( const TArray<FVertexAndAttributes>& Contour, TArray<FEdgeID>& OutExistingEdgeIDs, TArray<FEdgeID>& OutEdgeIDs, TArray<FMeshDescription::FContourPoint>& OutContourPoints );
+	void CreatePolygonContour( const TArray<FVertexAndAttributes>& Contour, TArray<FVertexInstanceID>& OutVertexInstanceIDs );
 	void BackupPolygonContour( const FMeshPolygonContour& Contour, TArray<FVertexAndAttributes>& OutVerticesAndAttributes );
 	void GetConnectedSoftEdges( const FVertexID VertexID, TArray<FEdgeID>& OutConnectedSoftEdges ) const;
 	void GetVertexConnectedPolygonsInSameSoftEdgedGroup( const FVertexID VertexInstanceID, const FPolygonID PolygonID, TArray<FPolygonID>& OutPolygonIDs ) const;
diff --git a/Engine/Plugins/Runtime/EditableMesh/Source/EditableMesh/Public/EditableMeshTypes.h b/Engine/Plugins/Runtime/EditableMesh/Source/EditableMesh/Public/EditableMeshTypes.h
--- a/Engine/Plugins/Runtime/EditableMesh/Source/EditableMesh/Public/EditableMeshTypes.h
+++ b/Engine/Plugins/Runtime/EditableMesh/Source/EditableMesh/Public/EditableMeshTypes.h
@@ -680,30 +680,6 @@ struct FVertexAndAttributes
 };
 
 
-USTRUCT( BlueprintType )
-struct FPolygonHoleVertices
-{
-	GENERATED_BODY()
-
-	/** Ordered list of vertices that defines the hole's contour, along with the polygon vertex attributes to set for each vertex */
-	UPROPERTY( BlueprintReadWrite, Category="Editable Mesh" )
-	TArray<FVertexAndAttributes> HoleVertices;
-
-	/** Default constructor */
-	FPolygonHoleVertices()
-		: HoleVertices()
-	{
-	}
-
-	FString ToString() const
-	{
-		return FString::Printf(
-			TEXT( "HoleVertices:%s" ),
-			*LogHelpers::ArrayToString( HoleVertices ) );
-	}
-};
-
-
 UENUM( BlueprintType )
 enum class EPolygonEdgeHardness : uint8
 {
@@ -734,10 +710,6 @@ struct FPolygonToCreate
 	UPROPERTY( BlueprintReadWrite, Category="Editable Mesh" )
 	TArray<FVertexAndAttributes> PerimeterVertices;
 	
-	/** For each hole in the polygon, an ordered list of vertices that defines that hole's boundary */
-	UPROPERTY( BlueprintReadWrite, Category="Editable Mesh" )
-	TArray<FPolygonHoleVertices> PolygonHoles;
-
 	/** The original ID of the polygon.  Should only be used by the undo system. */
 	UPROPERTY( BlueprintReadWrite, Category="Editable Mesh" )
 	FPolygonID OriginalPolygonID;
@@ -750,7 +722,6 @@ struct FPolygonToCreate
 	FPolygonToCreate()
 		: PolygonGroupID( FPolygonGroupID::Invalid ),
 		  PerimeterVertices(),
-		  PolygonHoles(),
 		  OriginalPolygonID( FPolygonID::Invalid ),
 		  PolygonEdgeHardness( EPolygonEdgeHardness::NewEdgesSoft )
 	{
@@ -759,10 +730,9 @@ struct FPolygonToCreate
 	FString ToString() const
 	{
 		return FString::Printf(
-			TEXT( "PolygonGroupID:%s, PerimeterVertices:%s, PolygonHoles:%s, OriginalPolygonID:%s" ),
+			TEXT( "PolygonGroupID:%s, PerimeterVertices:%s, OriginalPolygonID:%s" ),
 			*PolygonGroupID.ToString(),
 			*LogHelpers::ArrayToString( PerimeterVertices ),
-			*LogHelpers::ArrayToString( PolygonHoles ),
 			*OriginalPolygonID.ToString() );
 	}
 };
diff --git a/Engine/Plugins/Runtime/ProceduralMeshComponent/Source/ProceduralMeshComponentEditor/Private/ProceduralMeshComponentDetails.cpp b/Engine/Plugins/Runtime/ProceduralMeshComponent/Source/ProceduralMeshComponentEditor/Private/ProceduralMeshComponentDetails.cpp
--- a/Engine/Plugins/Runtime/ProceduralMeshComponent/Source/ProceduralMeshComponentEditor/Private/ProceduralMeshComponentDetails.cpp
+++ b/Engine/Plugins/Runtime/ProceduralMeshComponent/Source/ProceduralMeshComponentEditor/Private/ProceduralMeshComponentDetails.cpp
@@ -211,7 +211,9 @@ FReply FProceduralMeshComponentDetails::ClickedOnConvertToStaticMesh()
 				for (int32 TriIdx = 0; TriIdx < NumTri; TriIdx++)
 				{
 					FVertexID VertexIndexes[3];
-					FVertexInstanceID VertexInstanceIDs[3];
+					TArray<FVertexInstanceID> VertexInstanceIDs;
+					VertexInstanceIDs.SetNum(3);
+
 					for (int32 CornerIndex = 0; CornerIndex < 3; ++CornerIndex)
 					{
 						const int32 IndiceIndex = (TriIdx * 3) + CornerIndex;
@@ -220,32 +222,8 @@ FReply FProceduralMeshComponentDetails::ClickedOnConvertToStaticMesh()
 						VertexInstanceIDs[CornerIndex] = IndiceIndexToVertexInstanceID[IndiceIndex];
 					}
 
-					TArray<FMeshDescription::FContourPoint> Contours;
-					for (int32 CornerIndex = 0; CornerIndex < 3; ++CornerIndex)
-					{
-						int32 ContourPointIndex = Contours.AddDefaulted();
-						FMeshDescription::FContourPoint& ContourPoint = Contours[ContourPointIndex];
-						//Find the matching edge ID
-						uint32 CornerIndices[2];
-						CornerIndices[0] = (CornerIndex + 0) % 3;
-						CornerIndices[1] = (CornerIndex + 1) % 3;
-
-						FVertexID EdgeVertexIDs[2];
-						EdgeVertexIDs[0] = VertexIndexes[CornerIndices[0]];
-						EdgeVertexIDs[1] = VertexIndexes[CornerIndices[1]];
-
-						FEdgeID MatchEdgeId = MeshDescription.GetVertexPairEdge(EdgeVertexIDs[0], EdgeVertexIDs[1]);
-						if (MatchEdgeId == FEdgeID::Invalid)
-						{
-							MatchEdgeId = MeshDescription.CreateEdge(EdgeVertexIDs[0], EdgeVertexIDs[1]);
-							EdgeHardnesses[MatchEdgeId] = false;
-							EdgeCreaseSharpnesses[MatchEdgeId] = 0.0f;
-						}
-						ContourPoint.EdgeID = MatchEdgeId;
-						ContourPoint.VertexInstanceID = VertexInstanceIDs[CornerIndices[0]];
-					}
 					// Insert a polygon into the mesh
-					const FPolygonID NewPolygonID = MeshDescription.CreatePolygon(PolygonGroupID, Contours);
+					const FPolygonID NewPolygonID = MeshDescription.CreatePolygon(PolygonGroupID, VertexInstanceIDs);
 					//Triangulate the polygon
 					FMeshPolygon& Polygon = MeshDescription.GetPolygon(NewPolygonID);
 					MeshDescription.ComputePolygonTriangulation(NewPolygonID, Polygon.Triangles);
diff --git a/Engine/Source/Developer/MeshDescriptionOperations/Private/MeshDescriptionOperations.cpp b/Engine/Source/Developer/MeshDescriptionOperations/Private/MeshDescriptionOperations.cpp
--- a/Engine/Source/Developer/MeshDescriptionOperations/Private/MeshDescriptionOperations.cpp
+++ b/Engine/Source/Developer/MeshDescriptionOperations/Private/MeshDescriptionOperations.cpp
@@ -528,7 +528,8 @@ void FMeshDescriptionOperations::ConvertFromRawMesh(const FRawMesh& SourceRawMes
 			PolygonGroups.Add(PolygonGroupID);
 			MaterialIndexToPolygonGroup.Add(MaterialIndex, PolygonGroupID);
 		}
-		FVertexInstanceID TriangleVertexInstanceIDs[3];
+		TArray<FVertexInstanceID> TriangleVertexInstanceIDs;
+		TriangleVertexInstanceIDs.SetNum(3);
 		for (int32 Corner = 0; Corner < 3; ++Corner)
 		{
 			int32 VerticeIndex = VerticeIndexBase + Corner;
@@ -550,31 +551,7 @@ void FMeshDescriptionOperations::ConvertFromRawMesh(const FRawMesh& SourceRawMes
 			}
 		}
 
-		//Create the polygon edges
-		TArray<FMeshDescription::FContourPoint> Contours;
-		for (uint32 Corner = 0; Corner < 3; ++Corner)
-		{
-			int32 ContourPointIndex = Contours.AddDefaulted();
-			FMeshDescription::FContourPoint& ContourPoint = Contours[ContourPointIndex];
-			//Find the matching edge ID
-			int32 CornerIndices[2];
-			CornerIndices[0] = (Corner + 0) % 3;
-			CornerIndices[1] = (Corner + 1) % 3;
-
-			FVertexID EdgeVertexIDs[2];
-			EdgeVertexIDs[0] = DestinationMeshDescription.GetVertexInstanceVertex(FVertexInstanceID(TriangleVertexInstanceIDs[CornerIndices[0]]));
-			EdgeVertexIDs[1] = DestinationMeshDescription.GetVertexInstanceVertex(FVertexInstanceID(TriangleVertexInstanceIDs[CornerIndices[1]]));
-
-			FEdgeID MatchEdgeId = DestinationMeshDescription.GetVertexPairEdge(EdgeVertexIDs[0], EdgeVertexIDs[1]);
-			if (MatchEdgeId == FEdgeID::Invalid)
-			{
-				MatchEdgeId = DestinationMeshDescription.CreateEdge(EdgeVertexIDs[0], EdgeVertexIDs[1]);
-			}
-			ContourPoint.EdgeID = MatchEdgeId;
-			ContourPoint.VertexInstanceID = FVertexInstanceID(TriangleVertexInstanceIDs[CornerIndices[0]]);
-		}
-
-		const FPolygonID NewPolygonID = DestinationMeshDescription.CreatePolygon(PolygonGroupID, Contours);
+		const FPolygonID NewPolygonID = DestinationMeshDescription.CreatePolygon(PolygonGroupID, TriangleVertexInstanceIDs);
 		int32 NewTriangleIndex = DestinationMeshDescription.GetPolygonTriangles(NewPolygonID).AddDefaulted();
 		FMeshTriangle& NewTriangle = DestinationMeshDescription.GetPolygonTriangles(NewPolygonID)[NewTriangleIndex];
 		for (int32 Corner = 0; Corner < 3; ++Corner)
@@ -761,34 +738,14 @@ void FMeshDescriptionOperations::AppendMeshDescription(const FMeshDescription& S
 		FPolygonGroupID TargetPolygonGroupID = RemapPolygonGroup[SourcePolygon.PolygonGroupID];
 
 		int32 PolygonVertexCount = SourcePolygon.PerimeterContour.VertexInstanceIDs.Num();
-		//Create a contour
-		TArray<FMeshDescription::FContourPoint> Contours;
-		// Add the edges of this polygon
-		for (uint32 PolygonEdgeNumber = 0; PolygonEdgeNumber < (uint32)PolygonVertexCount; ++PolygonEdgeNumber)
+		TArray<FVertexInstanceID> VertexInstanceIDs;
+		VertexInstanceIDs.Reserve(PolygonVertexCount);
+		for (const FVertexInstanceID VertexInstanceID : SourcePolygon.PerimeterContour.VertexInstanceIDs)
 		{
-			int32 ContourPointIndex = Contours.AddDefaulted();
-			FMeshDescription::FContourPoint& ContourPoint = Contours[ContourPointIndex];
-			//Find the matching edge ID
-			uint32 CornerIndices[2];
-			CornerIndices[0] = (PolygonEdgeNumber + 0) % PolygonVertexCount;
-			CornerIndices[1] = (PolygonEdgeNumber + 1) % PolygonVertexCount;
-
-			FVertexID SourceVertexIDs[2];
-			SourceVertexIDs[0] = SourceMesh.GetVertexInstanceVertex(SourcePolygon.PerimeterContour.VertexInstanceIDs[CornerIndices[0]]);
-			SourceVertexIDs[1] = SourceMesh.GetVertexInstanceVertex(SourcePolygon.PerimeterContour.VertexInstanceIDs[CornerIndices[1]]);
-
-			FEdgeID SourceEdgeID = SourceMesh.GetVertexPairEdge(SourceVertexIDs[0], SourceVertexIDs[1]);
-			check(SourceEdgeID != FEdgeID::Invalid)
-			FEdgeID TargetEdgeID = SourceEdgeIDRemap[SourceEdgeID];
-			ContourPoint.EdgeID = TargetEdgeID;
-			ContourPoint.VertexInstanceID = SourceVertexInstanceIDRemap[SourcePolygon.PerimeterContour.VertexInstanceIDs[CornerIndices[0]]];
-
-			TargetEdgeCreaseSharpnesses[TargetEdgeID] = SourceEdgeCreaseSharpnesses[SourceEdgeID];
-			TargetEdgeHardnesses[TargetEdgeID] = SourceEdgeHardnesses[SourceEdgeID];
+			VertexInstanceIDs.Add(SourceVertexInstanceIDRemap[VertexInstanceID]);
 		}
-
 		// Insert a polygon into the mesh
-		const FPolygonID TargetPolygonID = TargetMesh.CreatePolygon(TargetPolygonGroupID, Contours);
+		const FPolygonID TargetPolygonID = TargetMesh.CreatePolygon(TargetPolygonGroupID, VertexInstanceIDs);
 		//Triangulate the polygon
 		FMeshPolygon& Polygon = TargetMesh.GetPolygon(TargetPolygonID);
 		TargetMesh.ComputePolygonTriangulation(TargetPolygonID, Polygon.Triangles);
diff --git a/Engine/Source/Developer/MeshMergeUtilities/Private/MeshMergeHelpers.cpp b/Engine/Source/Developer/MeshMergeUtilities/Private/MeshMergeHelpers.cpp
--- a/Engine/Source/Developer/MeshMergeUtilities/Private/MeshMergeHelpers.cpp
+++ b/Engine/Source/Developer/MeshMergeUtilities/Private/MeshMergeHelpers.cpp
@@ -322,7 +322,8 @@ void FMeshMergeHelpers::RetrieveMesh(USkeletalMeshComponent* SkeletalMeshCompone
 			for (uint32 SectionTriangleIndex = 0; SectionTriangleIndex < SkelMeshSection.NumTriangles; ++SectionTriangleIndex)
 			{
 				FVertexID VertexIndexes[3];
-				FVertexInstanceID VertexInstanceIDs[3];
+				TArray<FVertexInstanceID> VertexInstanceIDs;
+				VertexInstanceIDs.SetNum(3);
 				for (int32 CornerIndex = 0; CornerIndex < 3; ++CornerIndex, ++WedgeIndex)
 				{
 					const int32 VertexIndexForWedge = LODModel.IndexBuffer[SkelMeshSection.BaseIndex + WedgeIndex];
@@ -353,32 +354,7 @@ void FMeshMergeHelpers::RetrieveMesh(USkeletalMeshComponent* SkeletalMeshCompone
 					VertexInstanceColors[VertexInstanceID] = bPropagateVertexColours ? FVector4(FLinearColor(SoftVertex.Color)) : FVector4(1.0f, 1.0f, 1.0f);
 				}
 				//Create a polygon from this triangle
-				TArray<FMeshDescription::FContourPoint> Contours;
-				for (int32 CornerIndex = 0; CornerIndex < 3; ++CornerIndex)
-				{
-					int32 ContourPointIndex = Contours.AddDefaulted();
-					FMeshDescription::FContourPoint& ContourPoint = Contours[ContourPointIndex];
-					//Find the matching edge ID
-					uint32 CornerIndices[2];
-					CornerIndices[0] = (CornerIndex + 0) % 3;
-					CornerIndices[1] = (CornerIndex + 1) % 3;
-
-					FVertexID EdgeVertexIDs[2];
-					EdgeVertexIDs[0] = VertexIndexes[CornerIndices[0]];
-					EdgeVertexIDs[1] = VertexIndexes[CornerIndices[1]];
-
-					FEdgeID MatchEdgeId = RawMesh.GetVertexPairEdge(EdgeVertexIDs[0], EdgeVertexIDs[1]);
-					if (MatchEdgeId == FEdgeID::Invalid)
-					{
-						MatchEdgeId = RawMesh.CreateEdge(EdgeVertexIDs[0], EdgeVertexIDs[1]);
-						EdgeHardnesses[MatchEdgeId] = false;
-						EdgeCreaseSharpnesses[MatchEdgeId] = 0.0f;
-					}
-					ContourPoint.EdgeID = MatchEdgeId;
-					ContourPoint.VertexInstanceID = VertexInstanceIDs[CornerIndices[0]];
-				}
-				// Insert a polygon into the mesh
-				const FPolygonID NewPolygonID = RawMesh.CreatePolygon(SectionPolygonGroupID, Contours);
+				const FPolygonID NewPolygonID = RawMesh.CreatePolygon(SectionPolygonGroupID, VertexInstanceIDs);
 				//Triangulate the polygon
 				FMeshPolygon& Polygon = RawMesh.GetPolygon(NewPolygonID);
 				RawMesh.ComputePolygonTriangulation(NewPolygonID, Polygon.Triangles);
@@ -500,7 +476,9 @@ void FMeshMergeHelpers::ExportStaticMeshLOD(const FStaticMeshLODResources& Stati
 		check(CurrentPolygonGroupID != FPolygonGroupID::Invalid);
 
 		FVertexID VertexIDs[3];
-		FVertexInstanceID VertexInstanceIDs[3];
+		TArray<FVertexInstanceID> VertexInstanceIDs;
+		VertexInstanceIDs.SetNum(3);
+
 		for (int32 Corner = 0; Corner < 3; ++Corner)
 		{
 			int32 WedgeIndex = StaticMeshLOD.IndexBuffer.GetIndex(TriangleIndex * 3 + Corner);
@@ -534,32 +512,7 @@ void FMeshMergeHelpers::ExportStaticMeshLOD(const FStaticMeshLODResources& Stati
 			}
 		}
 		//Create a polygon from this triangle
-		TArray<FMeshDescription::FContourPoint> Contours;
-		for (int32 Corner = 0; Corner < 3; ++Corner)
-		{
-			int32 ContourPointIndex = Contours.AddDefaulted();
-			FMeshDescription::FContourPoint& ContourPoint = Contours[ContourPointIndex];
-			//Find the matching edge ID
-			uint32 CornerIndices[2];
-			CornerIndices[0] = (Corner + 0) % 3;
-			CornerIndices[1] = (Corner + 1) % 3;
-
-			FVertexID EdgeVertexIDs[2];
-			EdgeVertexIDs[0] = VertexIDs[CornerIndices[0]];
-			EdgeVertexIDs[1] = VertexIDs[CornerIndices[1]];
-
-			FEdgeID MatchEdgeId = OutRawMesh.GetVertexPairEdge(EdgeVertexIDs[0], EdgeVertexIDs[1]);
-			if (MatchEdgeId == FEdgeID::Invalid)
-			{
-				MatchEdgeId = OutRawMesh.CreateEdge(EdgeVertexIDs[0], EdgeVertexIDs[1]);
-				EdgeHardnesses[MatchEdgeId] = false;
-				EdgeCreaseSharpnesses[MatchEdgeId] = 0.0f;
-			}
-			ContourPoint.EdgeID = MatchEdgeId;
-			ContourPoint.VertexInstanceID = VertexInstanceIDs[CornerIndices[0]];
-		}
-		// Insert a polygon into the mesh
-		const FPolygonID NewPolygonID = OutRawMesh.CreatePolygon(CurrentPolygonGroupID, Contours);
+		const FPolygonID NewPolygonID = OutRawMesh.CreatePolygon(CurrentPolygonGroupID, VertexInstanceIDs);
 		//Triangulate the polygon
 		FMeshPolygon& Polygon = OutRawMesh.GetPolygon(NewPolygonID);
 		OutRawMesh.ComputePolygonTriangulation(NewPolygonID, Polygon.Triangles);
@@ -1286,22 +1239,13 @@ void FMeshMergeHelpers::AppendRawMesh(FMeshDescription& InTarget, const FMeshDes
 		SourceToTargetEdgeID.Add(SourceEdgeID, TargetEdgeID);
 	}
 
-	auto CreateContour = [&InSource, &SourceToTargetVertexInstanceID, &SourceToTargetEdgeID](const TArray<FVertexInstanceID>& SourceVertexInstanceIDs, TArray<FMeshDescription::FContourPoint>& Contours)
+	auto CreateContour = [&InSource, &SourceToTargetVertexInstanceID, &SourceToTargetEdgeID](const TArray<FVertexInstanceID>& SourceVertexInstanceIDs, TArray<FVertexInstanceID>& DestVertexInstanceIDs)
 	{
 		const int32 ContourCount = SourceVertexInstanceIDs.Num();
 		for (int32 ContourIndex = 0; ContourIndex < ContourCount; ++ContourIndex)
 		{
 			FVertexInstanceID SourceVertexInstanceID = SourceVertexInstanceIDs[ContourIndex];
-			int32 ContourPlusOne = (ContourIndex + 1) % ContourCount;
-			FVertexInstanceID SourceVertexInstanceID_Next = SourceVertexInstanceIDs[ContourPlusOne];
-
-			int32 ContourPointIndex = Contours.AddDefaulted();
-			FMeshDescription::FContourPoint& ContourPoint = Contours[ContourPointIndex];
-
-			const FEdgeID SourceEdgeID = InSource.GetVertexPairEdge(InSource.GetVertexInstanceVertex(SourceVertexInstanceID), InSource.GetVertexInstanceVertex(SourceVertexInstanceID_Next));
-			check(SourceEdgeID != FEdgeID::Invalid);
-			ContourPoint.VertexInstanceID = SourceToTargetVertexInstanceID[SourceVertexInstanceID];
-			ContourPoint.EdgeID = SourceToTargetEdgeID[SourceEdgeID];
+			DestVertexInstanceIDs.Add(SourceToTargetVertexInstanceID[SourceVertexInstanceID]);
 		}
 	};
 
@@ -1312,11 +1256,11 @@ void FMeshMergeHelpers::AppendRawMesh(FMeshDescription& InTarget, const FMeshDes
 		const TArray<FVertexInstanceID>& SourceVertexInstanceIDs = InSource.GetPolygonPerimeterVertexInstances(SourcePolygonID);
 
 
-		TArray<FMeshDescription::FContourPoint> Contours;
-		CreateContour(SourceVertexInstanceIDs, Contours);
+		TArray<FVertexInstanceID> ContourVertexInstances;
+		CreateContour(SourceVertexInstanceIDs, ContourVertexInstances);
 
 		// Insert a polygon into the mesh
-		const FPolygonID TargetPolygonID = InTarget.CreatePolygon(SourcePolygon.PolygonGroupID, Contours);
+		const FPolygonID TargetPolygonID = InTarget.CreatePolygon(SourcePolygon.PolygonGroupID, ContourVertexInstances);
 		//Triangulate the polygon
 		FMeshPolygon& Polygon = InTarget.GetPolygon(TargetPolygonID);
 		InTarget.ComputePolygonTriangulation(TargetPolygonID, Polygon.Triangles);
diff --git a/Engine/Source/Developer/MeshSimplifier/Private/QuadricMeshReduction.cpp b/Engine/Source/Developer/MeshSimplifier/Private/QuadricMeshReduction.cpp
--- a/Engine/Source/Developer/MeshSimplifier/Private/QuadricMeshReduction.cpp
+++ b/Engine/Source/Developer/MeshSimplifier/Private/QuadricMeshReduction.cpp
@@ -452,7 +452,9 @@ public:
 			//Vertex instances and Polygons
 			for (int32 TriangleIndex = 0; TriangleIndex < (int32)NumTris; TriangleIndex++)
 			{
-				FVertexInstanceID CornerInstanceIDs[3];
+				TArray<FVertexInstanceID> CornerInstanceIDs;
+				CornerInstanceIDs.SetNum(3);
+
 				FVertexID CornerVerticesIDs[3];
 				for (int32 CornerIndex = 0; CornerIndex < 3; CornerIndex++)
 				{
@@ -514,36 +516,13 @@ public:
 					MaterialPolygonGroupID = PolygonGroupMapping[MaterialIndex];
 				}
 
-				// Create polygon edges
-				TArray<FMeshDescription::FContourPoint> Contours;
+				// Insert a polygon into the mesh
+				TArray<FEdgeID> NewEdgeIDs;
+				const FPolygonID NewPolygonID = OutReducedMesh.CreatePolygon(MaterialPolygonGroupID, CornerInstanceIDs, &NewEdgeIDs);
+				for (const FEdgeID NewEdgeID : NewEdgeIDs)
 				{
-					// Add the edges of this triangle
-					for (uint32 TriangleEdgeNumber = 0; TriangleEdgeNumber < 3; ++TriangleEdgeNumber)
-					{
-						int32 ContourPointIndex = Contours.AddDefaulted();
-						FMeshDescription::FContourPoint& ContourPoint = Contours[ContourPointIndex];
-						//Find the matching edge ID
-						uint32 CornerIndices[2];
-						CornerIndices[0] = (TriangleEdgeNumber + 0) % 3;
-						CornerIndices[1] = (TriangleEdgeNumber + 1) % 3;
-
-						FVertexID EdgeVertexIDs[2];
-						EdgeVertexIDs[0] = CornerVerticesIDs[CornerIndices[0]];
-						EdgeVertexIDs[1] = CornerVerticesIDs[CornerIndices[1]];
-
-						FEdgeID MatchEdgeId = OutReducedMesh.GetVertexPairEdge(EdgeVertexIDs[0], EdgeVertexIDs[1]);
-						if (MatchEdgeId == FEdgeID::Invalid)
-						{
-							MatchEdgeId = OutReducedMesh.CreateEdge(EdgeVertexIDs[0], EdgeVertexIDs[1]);
-							// @todo: set edge hardness?
-						}
-						ContourPoint.EdgeID = MatchEdgeId;
-						ContourPoint.VertexInstanceID = CornerInstanceIDs[CornerIndices[0]];
-					}
+					// @todo: set NewEdgeID edge hardness?
 				}
-
-				// Insert a polygon into the mesh
-				const FPolygonID NewPolygonID = OutReducedMesh.CreatePolygon(MaterialPolygonGroupID, Contours);
 				const int32 NewTriangleIndex = OutReducedMesh.GetPolygonTriangles(NewPolygonID).AddDefaulted();
 				FMeshTriangle& NewTriangle = OutReducedMesh.GetPolygonTriangles(NewPolygonID)[NewTriangleIndex];
 				for (int32 TriangleVertexIndex = 0; TriangleVertexIndex < 3; ++TriangleVertexIndex)
diff --git a/Engine/Source/Editor/UnrealEd/Private/Fbx/FbxMainExport.cpp b/Engine/Source/Editor/UnrealEd/Private/Fbx/FbxMainExport.cpp
--- a/Engine/Source/Editor/UnrealEd/Private/Fbx/FbxMainExport.cpp
+++ b/Engine/Source/Editor/UnrealEd/Private/Fbx/FbxMainExport.cpp
@@ -1058,7 +1058,9 @@ void FFbxExporter::ExportBSP( UModel* Model, bool bSelectedOnly )
 
 			for (int32 StartVertexIndex = 1; StartVertexIndex < Node.NumVertices - 1; ++StartVertexIndex)
 			{
-				FVertexInstanceID VertexInstanceIDs[3];
+				TArray<FVertexInstanceID> VertexInstanceIDs;
+				VertexInstanceIDs.SetNum(3);
+
 				FVertexID VertexIDs[3];
 				// These map the node's vertices to the 3 triangle indices to triangulate the convex polygon.
 				int32 TriVertIndices[3] = {
@@ -1086,33 +1088,16 @@ void FFbxExporter::ExportBSP( UModel* Model, bool bSelectedOnly )
 					Colors[VertexInstanceIDs[Corner]] = FLinearColor::White;
 					Normals[VertexInstanceIDs[Corner]] = Normal;
 				}
-				//Add triangles
-				TArray<FMeshDescription::FContourPoint> Contours;
-				for (int32 Corner = 0; Corner < 3; ++Corner)
+
+				// Insert a polygon into the mesh
+				TArray<FEdgeID> NewEdgeIDs;
+				const FPolygonID NewPolygonID = Mesh.CreatePolygon(CurrentPolygonGroupID, VertexInstanceIDs, &NewEdgeIDs);
+				for (const FEdgeID EdgeID : NewEdgeIDs)
 				{
-					int32 ContourPointIndex = Contours.AddDefaulted();
-					FMeshDescription::FContourPoint& ContourPoint = Contours[ContourPointIndex];
-					//Find the matching edge ID
-					uint32 CornerIndices[2];
-					CornerIndices[0] = (Corner + 0) % 3;
-					CornerIndices[1] = (Corner + 1) % 3;
-
-					FVertexID EdgeVertexIDs[2];
-					EdgeVertexIDs[0] = VertexIDs[CornerIndices[0]];
-					EdgeVertexIDs[1] = VertexIDs[CornerIndices[1]];
-
-					FEdgeID MatchEdgeId = Mesh.GetVertexPairEdge(EdgeVertexIDs[0], EdgeVertexIDs[1]);
-					if (MatchEdgeId == FEdgeID::Invalid)
-					{
-						MatchEdgeId = Mesh.CreateEdge(EdgeVertexIDs[0], EdgeVertexIDs[1]);
-						EdgeHardnesses[MatchEdgeId] = false;
-						EdgeCreaseSharpnesses[MatchEdgeId] = 0.0f;
-					}
-					ContourPoint.EdgeID = MatchEdgeId;
-					ContourPoint.VertexInstanceID = VertexInstanceIDs[CornerIndices[0]];
+					EdgeHardnesses[EdgeID] = false;
+					EdgeCreaseSharpnesses[EdgeID] = 0.0f;
 				}
-				// Insert a polygon into the mesh
-				const FPolygonID NewPolygonID = Mesh.CreatePolygon(CurrentPolygonGroupID, Contours);
+
 				//Triangulate the polygon
 				FMeshPolygon& Polygon = Mesh.GetPolygon(NewPolygonID);
 				Mesh.ComputePolygonTriangulation(NewPolygonID, Polygon.Triangles);
diff --git a/Engine/Source/Editor/UnrealEd/Private/Fbx/FbxStaticMeshImport.cpp b/Engine/Source/Editor/UnrealEd/Private/Fbx/FbxStaticMeshImport.cpp
--- a/Engine/Source/Editor/UnrealEd/Private/Fbx/FbxStaticMeshImport.cpp
+++ b/Engine/Source/Editor/UnrealEd/Private/Fbx/FbxStaticMeshImport.cpp
@@ -834,13 +834,10 @@ bool UnFbx::FFbxImporter::BuildStaticMeshFromGeometry(FbxNode* Node, UStaticMesh
 		}
 
 		// Create polygon edges
-		TArray<FMeshDescription::FContourPoint> Contours;
 		{
 			// Add the edges of this polygon
 			for (uint32 PolygonEdgeNumber = 0; PolygonEdgeNumber < (uint32)PolygonVertexCount; ++PolygonEdgeNumber)
 			{
-				int32 ContourPointIndex = Contours.AddDefaulted();
-				FMeshDescription::FContourPoint& ContourPoint = Contours[ContourPointIndex];
 				//Find the matching edge ID
 				uint32 CornerIndices[2];
 				CornerIndices[0] = (PolygonEdgeNumber + 0) % PolygonVertexCount;
@@ -855,8 +852,7 @@ bool UnFbx::FFbxImporter::BuildStaticMeshFromGeometry(FbxNode* Node, UStaticMesh
 				{
 					MatchEdgeId = MeshDescription->CreateEdge(EdgeVertexIDs[0], EdgeVertexIDs[1]);
 				}
-				ContourPoint.EdgeID = MatchEdgeId;
-				ContourPoint.VertexInstanceID = CornerInstanceIDs[CornerIndices[0]];
+
 				//RawMesh do not have edges, so by ordering the edge with the triangle construction we can ensure back and forth conversion with RawMesh
 				//When raw mesh will be completly remove we can create the edges right after the vertex creation.
 				int32 EdgeIndex = INDEX_NONE;
@@ -896,7 +892,9 @@ bool UnFbx::FFbxImporter::BuildStaticMeshFromGeometry(FbxNode* Node, UStaticMesh
 		}
 		FPolygonGroupID PolygonGroupID = PolygonGroupMapping[RealMaterialIndex];
 		// Insert a polygon into the mesh
-		const FPolygonID NewPolygonID = MeshDescription->CreatePolygon(PolygonGroupID, Contours);
+		TArray<FEdgeID> NewEdgeIDs;
+		const FPolygonID NewPolygonID = MeshDescription->CreatePolygon(PolygonGroupID, CornerInstanceIDs, &NewEdgeIDs);
+		check(NewEdgeIDs.Num() == 0);
 		//Triangulate the polygon
 		FMeshPolygon& Polygon = MeshDescription->GetPolygon(NewPolygonID);
 		MeshDescription->ComputePolygonTriangulation(NewPolygonID, Polygon.Triangles);
diff --git a/Engine/Source/Editor/UnrealEd/Private/StaticMeshEdit.cpp b/Engine/Source/Editor/UnrealEd/Private/StaticMeshEdit.cpp
--- a/Engine/Source/Editor/UnrealEd/Private/StaticMeshEdit.cpp
+++ b/Engine/Source/Editor/UnrealEd/Private/StaticMeshEdit.cpp
@@ -762,7 +762,9 @@ void GetBrushMesh(ABrush* Brush, UModel* Model, FMeshDescription& MeshDescriptio
 			}
 
 			//Create the vertex instances
-			FVertexInstanceID VertexInstanceIDs[3];
+			TArray<FVertexInstanceID> VertexInstanceIDs;
+			VertexInstanceIDs.SetNum(3);
+
 			for (int32 CornerIndex = 0; CornerIndex < 3; ++CornerIndex)
 			{
 				if (VertexID[CornerIndex] == FVertexID::Invalid)
@@ -778,32 +780,13 @@ void GetBrushMesh(ABrush* Brush, UModel* Model, FMeshDescription& MeshDescriptio
 
 			// Create a polygon with the 3 vertex instances
 			
-			TArray<FMeshDescription::FContourPoint> Contours;
-			for (uint32 TriangleEdgeNumber = 0; TriangleEdgeNumber < 3; ++TriangleEdgeNumber)
+			TArray<FEdgeID> NewEdgeIDs;
+			const FPolygonID NewPolygonID = MeshDescription.CreatePolygon(CurrentPolygonGroupID, VertexInstanceIDs, &NewEdgeIDs);
+			for (const FEdgeID NewEdgeID : NewEdgeIDs)
 			{
-				int32 ContourPointIndex = Contours.AddDefaulted();
-				FMeshDescription::FContourPoint& ContourPoint = Contours[ContourPointIndex];
-				//Find the matching edge ID
-				uint32 CornerIndices[2];
-				CornerIndices[0] = (TriangleEdgeNumber + 0) % 3;
-				CornerIndices[1] = (TriangleEdgeNumber + 1) % 3;
-
-				FVertexID EdgeVertexIDs[2];
-				EdgeVertexIDs[0] = VertexID[CornerIndices[0]];
-				EdgeVertexIDs[1] = VertexID[CornerIndices[1]];
-
-				FEdgeID MatchEdgeId = MeshDescription.GetVertexPairEdge(EdgeVertexIDs[0], EdgeVertexIDs[1]);
-				if (MatchEdgeId == FEdgeID::Invalid)
-				{
-					MatchEdgeId = MeshDescription.CreateEdge(EdgeVertexIDs[0], EdgeVertexIDs[1]);
-				}
-				ContourPoint.EdgeID = MatchEdgeId;
-				ContourPoint.VertexInstanceID = VertexInstanceIDs[CornerIndices[0]];
 				//All edge are hard for BSP
-				EdgeHardnesses[MatchEdgeId] = true;
+				EdgeHardnesses[NewEdgeID] = true;
 			}
-			// Insert a polygon into the mesh
-			const FPolygonID NewPolygonID = MeshDescription.CreatePolygon(CurrentPolygonGroupID, Contours);
 			int32 NewTriangleIndex = MeshDescription.GetPolygonTriangles(NewPolygonID).AddDefaulted();
 			FMeshTriangle& NewTriangle = MeshDescription.GetPolygonTriangles(NewPolygonID)[NewTriangleIndex];
 			for (int32 TriangleVertexIndex = 0; TriangleVertexIndex < 3; ++TriangleVertexIndex)
diff --git a/Engine/Source/Runtime/Landscape/Private/LandscapeEdit.cpp b/Engine/Source/Runtime/Landscape/Private/LandscapeEdit.cpp
--- a/Engine/Source/Runtime/Landscape/Private/LandscapeEdit.cpp
+++ b/Engine/Source/Runtime/Landscape/Private/LandscapeEdit.cpp
@@ -2982,32 +2982,20 @@ bool ALandscapeProxy::ExportToRawMesh(int32 InExportLOD, FMeshDescription& OutRa
 							auto AddTriangle = [&OutRawMesh, &EdgeHardnesses, &EdgeCreaseSharpnesses, &PolygonGroupID, &VertexIDs, &VertexInstanceIDs](int32 BaseIndex)
 							{
 								//Create a polygon from this triangle
-								TArray<FMeshDescription::FContourPoint> Contours;
+								TArray<FVertexInstanceID> PerimeterVertexInstances;
+								PerimeterVertexInstances.SetNum(3);
 								for (int32 Corner = 0; Corner < 3; ++Corner)
 								{
-									int32 ContourPointIndex = Contours.AddDefaulted();
-									FMeshDescription::FContourPoint& ContourPoint = Contours[ContourPointIndex];
-									//Find the matching edge ID
-									uint32 CornerIndices[2];
-									CornerIndices[0] = BaseIndex + ((Corner + 0) % 3);
-									CornerIndices[1] = BaseIndex + ((Corner + 1) % 3);
-
-									FVertexID EdgeVertexIDs[2];
-									EdgeVertexIDs[0] = VertexIDs[CornerIndices[0]];
-									EdgeVertexIDs[1] = VertexIDs[CornerIndices[1]];
-
-									FEdgeID MatchEdgeId = OutRawMesh.GetVertexPairEdge(EdgeVertexIDs[0], EdgeVertexIDs[1]);
-									if (MatchEdgeId == FEdgeID::Invalid)
-									{
-										MatchEdgeId = OutRawMesh.CreateEdge(EdgeVertexIDs[0], EdgeVertexIDs[1]);
-										EdgeHardnesses[MatchEdgeId] = false;
-										EdgeCreaseSharpnesses[MatchEdgeId] = 0.0f;
-									}
-									ContourPoint.EdgeID = MatchEdgeId;
-									ContourPoint.VertexInstanceID = VertexInstanceIDs[CornerIndices[0]];
+									PerimeterVertexInstances[Corner] = VertexInstanceIDs[BaseIndex + Corner];
 								}
 								// Insert a polygon into the mesh
-								const FPolygonID NewPolygonID = OutRawMesh.CreatePolygon(PolygonGroupID, Contours);
+								TArray<FEdgeID> NewEdgeIDs;
+								const FPolygonID NewPolygonID = OutRawMesh.CreatePolygon(PolygonGroupID, PerimeterVertexInstances, &NewEdgeIDs);
+								for (const FEdgeID NewEdgeID : NewEdgeIDs)
+								{
+									EdgeHardnesses[NewEdgeID] = false;
+									EdgeCreaseSharpnesses[NewEdgeID] = 0.0f;
+								}
 								//Triangulate the polygon
 								FMeshPolygon& Polygon = OutRawMesh.GetPolygon(NewPolygonID);
 								OutRawMesh.ComputePolygonTriangulation(NewPolygonID, Polygon.Triangles);
diff --git a/Engine/Source/Runtime/MeshDescription/Private/MeshDescription.cpp b/Engine/Source/Runtime/MeshDescription/Private/MeshDescription.cpp
--- a/Engine/Source/Runtime/MeshDescription/Private/MeshDescription.cpp
+++ b/Engine/Source/Runtime/MeshDescription/Private/MeshDescription.cpp
@@ -229,6 +229,76 @@ void FMeshDescription::FixUpElementIDs( const FElementIDRemappings& Remappings )
 }
 
 
+void FMeshDescription::CreatePolygon_Internal( const FPolygonID PolygonID, const FPolygonGroupID PolygonGroupID, const TArray<FContourPoint>& Perimeter )
+{
+	FMeshPolygon& Polygon = PolygonArray[ PolygonID ];
+
+	Polygon.PerimeterContour.VertexInstanceIDs.Reset( Perimeter.Num() );
+	for( const FContourPoint ContourPoint : Perimeter )
+	{
+		const FVertexInstanceID VertexInstanceID = ContourPoint.VertexInstanceID;
+		const FEdgeID EdgeID = ContourPoint.EdgeID;
+
+		Polygon.PerimeterContour.VertexInstanceIDs.Add( VertexInstanceID );
+		check( !VertexInstanceArray[ VertexInstanceID ].ConnectedPolygons.Contains( PolygonID ) );
+		VertexInstanceArray[ VertexInstanceID ].ConnectedPolygons.Add( PolygonID );
+
+		check( !EdgeArray[ EdgeID ].ConnectedPolygons.Contains( PolygonID ) );
+		EdgeArray[ EdgeID ].ConnectedPolygons.Add( PolygonID );
+	}
+
+	Polygon.PolygonGroupID = PolygonGroupID;
+	PolygonGroupArray[ PolygonGroupID ].Polygons.Add( PolygonID );
+
+	PolygonAttributesSet.Insert( PolygonID );
+}
+
+
+void FMeshDescription::CreatePolygon_Internal( const FPolygonID PolygonID, const FPolygonGroupID PolygonGroupID, const TArray<FVertexInstanceID>& VertexInstanceIDs, TArray<FEdgeID>* OutEdgeIDs )
+{
+	if( OutEdgeIDs )
+	{
+		OutEdgeIDs->Empty();
+	}
+
+	FMeshPolygon& Polygon = PolygonArray[ PolygonID ];
+	const int32 NumVertices = VertexInstanceIDs.Num();
+	Polygon.PerimeterContour.VertexInstanceIDs.Reset( NumVertices );
+
+	for( int32 Index = 0; Index < NumVertices; ++Index )
+	{
+		const FVertexInstanceID VertexInstanceID = VertexInstanceIDs[ Index ];
+		const FVertexInstanceID NextVertexInstanceID = VertexInstanceIDs[ ( Index + 1 == NumVertices ) ? 0 : Index + 1 ];
+
+		Polygon.PerimeterContour.VertexInstanceIDs.Add( VertexInstanceID );
+		check( !VertexInstanceArray[ VertexInstanceID ].ConnectedPolygons.Contains( PolygonID ) );
+		VertexInstanceArray[ VertexInstanceID ].ConnectedPolygons.Add( PolygonID );
+
+		const FVertexID VertexID0 = GetVertexInstanceVertex( VertexInstanceID );
+		const FVertexID VertexID1 = GetVertexInstanceVertex( NextVertexInstanceID );
+
+		FEdgeID EdgeID = GetVertexPairEdge( VertexID0, VertexID1 );
+		if( EdgeID == FEdgeID::Invalid )
+		{
+			EdgeID = CreateEdge( VertexID0, VertexID1 );
+			if( OutEdgeIDs )
+			{
+				OutEdgeIDs->Add( EdgeID );
+			}
+		}
+
+		check( !EdgeArray[ EdgeID ].ConnectedPolygons.Contains( PolygonID ) );
+		EdgeArray[ EdgeID ].ConnectedPolygons.Add( PolygonID );
+	}
+
+	Polygon.PolygonGroupID = PolygonGroupID;
+	PolygonGroupArray[ PolygonGroupID ].Polygons.Add( PolygonID );
+
+	PolygonAttributesSet.Insert( PolygonID );
+}
+
+
+
 void FMeshDescription::RemapAttributes( const FElementIDRemappings& Remappings )
 {
 	VertexAttributesSet.Remap( Remappings.NewVertexIndexLookup );
diff --git a/Engine/Source/Runtime/MeshDescription/Public/MeshDescription.h b/Engine/Source/Runtime/MeshDescription/Public/MeshDescription.h
--- a/Engine/Source/Runtime/MeshDescription/Public/MeshDescription.h
+++ b/Engine/Source/Runtime/MeshDescription/Public/MeshDescription.h
@@ -508,36 +508,12 @@ public:
 	};
 
 private:
-	void CreatePolygonContour_Internal( const FPolygonID PolygonID, const TArray<FContourPoint>& ContourPoints, FMeshPolygonContour& Contour )
-	{
-		Contour.VertexInstanceIDs.Reset( ContourPoints.Num() );
-		for( const FContourPoint ContourPoint : ContourPoints )
-		{
-			const FVertexInstanceID VertexInstanceID = ContourPoint.VertexInstanceID;
-			const FEdgeID EdgeID = ContourPoint.EdgeID;
-
-			Contour.VertexInstanceIDs.Add( VertexInstanceID );
-			check( !VertexInstanceArray[ VertexInstanceID ].ConnectedPolygons.Contains( PolygonID ) );
-			VertexInstanceArray[ VertexInstanceID ].ConnectedPolygons.Add( PolygonID );
-
-			check( !EdgeArray[ EdgeID ].ConnectedPolygons.Contains( PolygonID ) );
-			EdgeArray[ EdgeID ].ConnectedPolygons.Add( PolygonID );
-		}
-	}
-
-	void CreatePolygon_Internal( const FPolygonID PolygonID, const FPolygonGroupID PolygonGroupID, const TArray<FContourPoint>& Perimeter )
-	{
-		FMeshPolygon& Polygon = PolygonArray[ PolygonID ];
-		CreatePolygonContour_Internal( PolygonID, Perimeter, Polygon.PerimeterContour );
-
-		Polygon.PolygonGroupID = PolygonGroupID;
-		PolygonGroupArray[ PolygonGroupID ].Polygons.Add( PolygonID );
-
-		PolygonAttributesSet.Insert( PolygonID );
-	}
+	void CreatePolygon_Internal( const FPolygonID PolygonID, const FPolygonGroupID PolygonGroupID, const TArray<FContourPoint>& Perimeter );
+	void CreatePolygon_Internal( const FPolygonID PolygonID, const FPolygonGroupID PolygonGroupID, const TArray<FVertexInstanceID>& VertexInstanceIDs, TArray<FEdgeID>* OutEdgeIDs );
 
 public:
 	/** Adds a new polygon to the mesh and returns its ID */
+	UE_DEPRECATED(4.22, "Please use a different CreatePolygon() overload.")
 	FPolygonID CreatePolygon( const FPolygonGroupID PolygonGroupID, const TArray<FContourPoint>& Perimeter )
 	{
 		const FPolygonID PolygonID = PolygonArray.Add();
@@ -545,13 +521,27 @@ public:
 		return PolygonID;
 	}
 
+	FPolygonID CreatePolygon( const FPolygonGroupID PolygonGroupID, const TArray<FVertexInstanceID>& VertexInstanceIDs, TArray<FEdgeID>* OutEdgeIDs = nullptr )
+	{
+		const FPolygonID PolygonID = PolygonArray.Add();
+		CreatePolygon_Internal( PolygonID, PolygonGroupID, VertexInstanceIDs, OutEdgeIDs );
+		return PolygonID;
+	}
+
 	/** Adds a new polygon to the mesh with the given ID */
+	UE_DEPRECATED(4.22, "Please use a different CreatePolygonWithID() overload.")
 	void CreatePolygonWithID( const FPolygonID PolygonID, const FPolygonGroupID PolygonGroupID, const TArray<FContourPoint>& Perimeter )
 	{
 		PolygonArray.Insert( PolygonID );
 		CreatePolygon_Internal( PolygonID, PolygonGroupID, Perimeter );
 	}
 
+	void CreatePolygonWithID( const FPolygonID PolygonID, const FPolygonGroupID PolygonGroupID, const TArray<FVertexInstanceID>& VertexInstanceIDs, TArray<FEdgeID>* OutEdgeIDs = nullptr )
+	{
+		PolygonArray.Insert( PolygonID );
+		CreatePolygon_Internal( PolygonID, PolygonGroupID, VertexInstanceIDs, OutEdgeIDs );
+	}
+
 private:
 	void DeletePolygonContour_Internal( const FPolygonID PolygonID, const TArray<FVertexInstanceID>& VertexInstanceIDs, TArray<FEdgeID>* InOutOrphanedEdgesPtr, TArray<FVertexInstanceID>* InOutOrphanedVertexInstancesPtr )
 	{
