Changed the "Flip Normals" operation to only reverse the normals/tangents and retriangulate instead of deleting/recreating the polygons and recomputing the normals/tangents.
Fixed some polygons didn't recompute their normals when "Unify Normals".

diff --git a/Engine/Plugins/Editor/MeshEditor/Source/PolygonModeling/UnifyNormals.cpp b/Engine/Plugins/Editor/MeshEditor/Source/PolygonModeling/UnifyNormals.cpp
--- a/Engine/Plugins/Editor/MeshEditor/Source/PolygonModeling/UnifyNormals.cpp
+++ b/Engine/Plugins/Editor/MeshEditor/Source/PolygonModeling/UnifyNormals.cpp
@@ -88,6 +88,8 @@ void UUnifyNormalsCommand::Execute(IMeshEditorModeEditingContract& MeshEditorMod
 		TSet<FPolygonID> SelectedPolygons;		// set of selected polygons
 		TSet<FPolygonID> FlippedPolygons;		// set of polygons that are flipped with respect to the selected polygons
 		TSet<FPolygonID> BoundaryPolygons;		// set of polygons that represent the boundary of the flipped polygons
+		TSet<FVertexID> BoundaryPolygonVertices;// set of vertices of the polygons at the boundary of the flipped polygons
+		TArray<FVertexID> Vertices;
 
 		// Check the neighbor polygons of the selected polygon to see if they are flipped with respect to it
 		// They are flipped if the edge shared by the polygons have the same edge direction (because they have opposite winding)
@@ -134,7 +136,10 @@ void UUnifyNormalsCommand::Execute(IMeshEditorModeEditingContract& MeshEditorMod
 								}
 								else
 								{
-									BoundaryPolygons.Add(NeighborPolygonID);
+									// Add the vertices of the polygon at the boundary
+									Vertices.Reset();
+									MeshDescription->GetPolygonPerimeterVertices(NeighborPolygonID, Vertices);
+									BoundaryPolygonVertices.Append(Vertices);
 								}
 							}
 						}
@@ -150,11 +155,34 @@ void UUnifyNormalsCommand::Execute(IMeshEditorModeEditingContract& MeshEditorMod
 
 		if (FlippedPolygons.Num() > 0)
 		{
+			// Get the set of vertices of all flipped polygons
+			TSet<FVertexID> FlippedPolygonVertices;
+			for (const FPolygonID& PolygonID : FlippedPolygons)
+			{
+				Vertices.Reset();
+				MeshDescription->GetPolygonPerimeterVertices(PolygonID, Vertices);
+				FlippedPolygonVertices.Append(Vertices);
+			}
+
+			// Get the vertices that are at the boundary of flipped and unflipped polygons
+			BoundaryPolygonVertices = BoundaryPolygonVertices.Intersect(FlippedPolygonVertices);
+
+			// The boundary polygons are those that share the boundary vertices
+			TSet<FPolygonID> BoundaryPolygons;
+			for (const FVertexID& VertexID : BoundaryPolygonVertices)
+			{
+				TArray<FPolygonID> Polygons;
+				MeshDescription->GetVertexConnectedPolygons(VertexID, Polygons);
+				BoundaryPolygons.Append(Polygons);
+			}
+
 			EditableMesh->StartModification(EMeshModificationType::Final, EMeshTopologyChange::TopologyChange);
 
 			EditableMesh->FlipPolygons(FlippedPolygons.Array());
 
 			// The selected and boundary polygons are not modified, but need to have their normals/tangents recomputed
+			// Also, adding the flipped polygons will force their tangents to be recomputed instead of just being flipped
+			EditableMesh->PolygonsPendingNewTangentBasis.Append(FlippedPolygons);
 			EditableMesh->PolygonsPendingNewTangentBasis.Append(SelectedPolygons);
 			EditableMesh->PolygonsPendingNewTangentBasis.Append(BoundaryPolygons);
 
diff --git a/Engine/Plugins/Runtime/EditableMesh/Source/EditableMesh/EditableMesh.cpp b/Engine/Plugins/Runtime/EditableMesh/Source/EditableMesh/EditableMesh.cpp
--- a/Engine/Plugins/Runtime/EditableMesh/Source/EditableMesh/EditableMesh.cpp
+++ b/Engine/Plugins/Runtime/EditableMesh/Source/EditableMesh/EditableMesh.cpp
@@ -652,6 +652,7 @@ void UEditableMesh::StartModification( const EMeshModificationType MeshModificat
 		}
 
 		PolygonsPendingNewTangentBasis.Reset();
+		PolygonsPendingFlipTangentBasis.Reset();
 		PolygonsPendingTriangulation.Reset();
 		VerticesPendingMerging.Reset();
 
@@ -693,6 +694,13 @@ void UEditableMesh::EndModification( const bool bFromUndo )
 				GenerateTangentsAndNormals();
 			}
 
+			// Exclude the polygons that have already regenerated their normal/tangent
+			PolygonsPendingFlipTangentBasis = PolygonsPendingFlipTangentBasis.Difference( PolygonsPendingNewTangentBasis );
+			if( PolygonsPendingFlipTangentBasis.Num() > 0 )
+			{
+				FlipTangentsAndNormals();
+			}
+
 			bAllowUndo = bIsUndoAllowed;
 		}
 
@@ -6175,6 +6183,44 @@ void UEditableMesh::GenerateTangentsAndNormals()
 	SetVertexInstancesAttributes( AttributesForVertexInstances );
 }
 
+void UEditableMesh::FlipTangentsAndNormals()
+{
+	const TVertexInstanceAttributesRef<FVector> VertexNormals = GetMeshDescription()->VertexInstanceAttributes().GetAttributesRef<FVector>( MeshAttribute::VertexInstance::Normal );
+	const TVertexInstanceAttributesRef<FVector> VertexTangents = GetMeshDescription()->VertexInstanceAttributes().GetAttributesRef<FVector>( MeshAttribute::VertexInstance::Tangent );
+	const TVertexInstanceAttributesRef<float> VertexBinormals = GetMeshDescription()->VertexInstanceAttributes().GetAttributesRef<float>( MeshAttribute::VertexInstance::BinormalSign );
+
+	static TSet<FVertexInstanceID> VertexInstanceIDs;
+	VertexInstanceIDs.Reset();
+
+	for( const FPolygonID PolygonID : PolygonsPendingFlipTangentBasis )
+	{
+		VertexInstanceIDs.Append( GetMeshDescription()->GetPolygonPerimeterVertexInstances( PolygonID ) );
+	}
+
+	static TArray<FAttributesForVertexInstance> AttributesForVertexInstances;
+	AttributesForVertexInstances.Reset( VertexInstanceIDs.Num() );
+
+	for( const FVertexInstanceID VertexInstanceID : VertexInstanceIDs )
+	{
+		const FVertexID VertexID = GetMeshDescription()->GetVertexInstanceVertex( VertexInstanceID );
+
+		// Just reverse the sign of the normals/tangents; note that since binormals are the cross product of normal with tangent, they are left untouched
+		FVector Normal = VertexNormals[ VertexInstanceID ] * -1.0f;
+		FVector Tangent = VertexTangents[ VertexInstanceID ] * -1.0f;
+		float BinormalSign = VertexBinormals[ VertexInstanceID ];
+
+		FAttributesForVertexInstance& AttributesForVertexInstance = AttributesForVertexInstances.AddDefaulted_GetRef();
+
+		AttributesForVertexInstance.VertexInstanceID = VertexInstanceID;
+		AttributesForVertexInstance.VertexInstanceAttributes.Attributes.Reset( 3 );
+		AttributesForVertexInstance.VertexInstanceAttributes.Attributes.Emplace( MeshAttribute::VertexInstance::Normal, 0, FMeshElementAttributeValue( Normal ) );
+		AttributesForVertexInstance.VertexInstanceAttributes.Attributes.Emplace( MeshAttribute::VertexInstance::Tangent, 0, FMeshElementAttributeValue( Tangent ) );
+		AttributesForVertexInstance.VertexInstanceAttributes.Attributes.Emplace( MeshAttribute::VertexInstance::BinormalSign, 0, FMeshElementAttributeValue( BinormalSign ) );
+	}
+
+	SetVertexInstancesAttributes( AttributesForVertexInstances );
+}
+
 void UEditableMesh::SplitPolygonalMesh(const FPlane& InPlane, TArray<FPolygonID>& PolygonIDs1, TArray<FPolygonID>& PolygonIDs2, TArray<FEdgeID>& BoundaryEdges)
 {
 	TPolygonAttributesConstRef<FVector> PolygonCenters = GetMeshDescription()->PolygonAttributes().GetAttributesRef<FVector>(MeshAttribute::Polygon::Center);
@@ -6851,51 +6897,22 @@ void UEditableMesh::FlipPolygons( const TArray<FPolygonID>& PolygonIDs )
 {
 	EM_ENTER( TEXT( "FlipPolygons %s" ), *LogHelpers::ArrayToString( PolygonIDs ) );
 
-	static TArray<FPolygonToCreate> PolygonsToCreate;
-	PolygonsToCreate.Reset();
-	for( int32 PolygonNumber = 0; PolygonNumber < PolygonIDs.Num(); ++PolygonNumber )
-	{
-		const FPolygonID OriginalPolygonID = PolygonIDs[ PolygonNumber ];
-		const FPolygonGroupID OriginalPolygonGroupID = GetGroupForPolygon( OriginalPolygonID );
-
-		FPolygonToCreate& PolygonToCreate = *new( PolygonsToCreate ) FPolygonToCreate();
-
-		PolygonToCreate.PolygonGroupID = OriginalPolygonGroupID;
-		
-		// Keep the original polygon ID.  No reason not to.
-		PolygonToCreate.OriginalPolygonID = OriginalPolygonID;
+	FFlipPolygonsChangeInput RevertInput;
+	RevertInput.PolygonIDsToFlip = PolygonIDs;
+	AddUndo( MakeUnique<FFlipPolygonsChange>( MoveTemp( RevertInput ) ) );
 
-		// Iterate backwards to add the vertices for the polygon, because we're flipping it over.
-		const TArray<FVertexInstanceID>& PerimeterVertexInstanceIDs = GetMeshDescription()->GetPolygonPerimeterVertexInstances( OriginalPolygonID );
-		const int32 PerimeterVertexCount = PerimeterVertexInstanceIDs.Num();
-
-		PolygonToCreate.PerimeterVertices.Reserve( PerimeterVertexCount );
-		for( int32 VertexNumber = PerimeterVertexCount - 1; VertexNumber >= 0; --VertexNumber )
-		{
-			PolygonToCreate.PerimeterVertices.Emplace();
-			FVertexAndAttributes& PerimeterVertex = PolygonToCreate.PerimeterVertices.Last();
-
-			// Just create a polygon from existing vertex instance IDs
-			PerimeterVertex.VertexInstanceID = PerimeterVertexInstanceIDs[ VertexNumber ];
-		}
-	}
-
-	// Delete all of the polygons
+	for (const FPolygonID& PolygonID : PolygonIDs)
 	{
-		const bool bDeleteOrphanedEdges = false;
-		const bool bDeleteOrphanedVertices = false;
-		const bool bDeleteOrphanedVertexInstances = false;
-		const bool bDeleteEmptyPolygonGroups = false;
-		DeletePolygons( PolygonIDs, bDeleteOrphanedEdges, bDeleteOrphanedVertices, bDeleteOrphanedVertexInstances, bDeleteEmptyPolygonGroups );
+		GetMeshDescription()->ReversePolygonFacing(PolygonID);
 	}
 
-	// Create the new polygons.  They're just like the polygons we deleted, except with reversed winding.
+	for( UEditableMeshAdapter* Adapter : Adapters )
 	{
-		static TArray<FPolygonID> NewPolygonIDs;
-		static TArray<FEdgeID> NewEdgeIDs;
-		CreatePolygons( PolygonsToCreate, NewPolygonIDs, NewEdgeIDs );
+		Adapter->OnRetriangulatePolygons( this, PolygonIDs );
 	}
 
+	PolygonsPendingFlipTangentBasis.Append(PolygonIDs);
+
 	EM_EXIT( TEXT( "FlipPolygons returned" ) );
 }
 
diff --git a/Engine/Plugins/Runtime/EditableMesh/Source/EditableMesh/EditableMeshChanges.cpp b/Engine/Plugins/Runtime/EditableMesh/Source/EditableMesh/EditableMeshChanges.cpp
--- a/Engine/Plugins/Runtime/EditableMesh/Source/EditableMesh/EditableMeshChanges.cpp
+++ b/Engine/Plugins/Runtime/EditableMesh/Source/EditableMesh/EditableMeshChanges.cpp
@@ -163,6 +163,22 @@ FString FDeletePolygonsChange::ToString() const
 		*LogHelpers::BoolToString( Input.bDeleteEmptySections ) );
 }
 
+TUniquePtr<FChange> FFlipPolygonsChange::Execute( UObject* Object )
+{
+	UEditableMesh* EditableMesh = CastChecked<UEditableMesh>( Object );
+	verify( !EditableMesh->AnyChangesToUndo() );
+
+	EditableMesh->FlipPolygons( Input.PolygonIDsToFlip );
+	return EditableMesh->MakeUndo();
+}
+
+
+FString FFlipPolygonsChange::ToString() const
+{
+	return FString::Printf(
+		TEXT( "Flip Polygons [PolygonIDsToFlip:%s]" ),
+		*LogHelpers::ArrayToString( Input.PolygonIDsToFlip ) );
+}
 
 TUniquePtr<FChange> FSetVerticesAttributesChange::Execute( UObject* Object )
 {
diff --git a/Engine/Plugins/Runtime/EditableMesh/Source/EditableMesh/EditableMeshChanges.h b/Engine/Plugins/Runtime/EditableMesh/Source/EditableMesh/EditableMeshChanges.h
--- a/Engine/Plugins/Runtime/EditableMesh/Source/EditableMesh/EditableMeshChanges.h
+++ b/Engine/Plugins/Runtime/EditableMesh/Source/EditableMesh/EditableMeshChanges.h
@@ -365,6 +365,34 @@ private:
 
 };
 
+struct FFlipPolygonsChangeInput
+{
+	/** The polygons to flip */
+	TArray<FPolygonID> PolygonIDsToFlip;
+
+	FFlipPolygonsChangeInput()
+		: PolygonIDsToFlip()
+	{
+	}
+};
+
+class FFlipPolygonsChange : public FChange
+{
+public:
+	FFlipPolygonsChange( const FFlipPolygonsChangeInput& InitInput )
+		: Input( InitInput )
+	{
+	}
+
+	// Parent class overrides
+	virtual TUniquePtr<FChange> Execute( UObject* Object ) override;
+	virtual FString ToString() const override;
+
+private:
+
+	/** The data we need to make this change */
+	FFlipPolygonsChangeInput Input;
+};
 
 struct FSetVerticesAttributesChangeInput
 {
diff --git a/Engine/Plugins/Runtime/EditableMesh/Source/EditableMesh/Public/EditableMesh.h b/Engine/Plugins/Runtime/EditableMesh/Source/EditableMesh/Public/EditableMesh.h
--- a/Engine/Plugins/Runtime/EditableMesh/Source/EditableMesh/Public/EditableMesh.h
+++ b/Engine/Plugins/Runtime/EditableMesh/Source/EditableMesh/Public/EditableMesh.h
@@ -503,6 +503,9 @@ private:
 	/** Called during end modification to generate tangents and normals on the pending polygon list */
 	void GenerateTangentsAndNormals();
 
+	/** Called during end modification to flip tangents and normals on the pending polygon list */
+	void FlipTangentsAndNormals();
+
 	/** Called during end modification to retriangulate polygons in the pending polygon list */
 	void RetriangulatePolygons();
 
@@ -556,6 +559,9 @@ public:
 	/** List of polygons which need their tangent basis recalculating (and consequently their associated vertex instances) */
 	TSet<FPolygonID> PolygonsPendingNewTangentBasis;
 
+	/** List of polygons which need their tangent basis flipped (and consequently their associated vertex instances) */
+	TSet<FPolygonID> PolygonsPendingFlipTangentBasis;
+
 	/** List of polygons requiring retriangulation */
 	TSet<FPolygonID> PolygonsPendingTriangulation;
 
