Disabled simple collision regeneration when editing static mesh in Edit Mode in the Static Mesh Editor.

diff --git a/Engine/Plugins/Runtime/EditableMesh/Source/EditableMesh/EditableStaticMeshAdapter.cpp b/Engine/Plugins/Runtime/EditableMesh/Source/EditableMesh/EditableStaticMeshAdapter.cpp
--- a/Engine/Plugins/Runtime/EditableMesh/Source/EditableMesh/EditableStaticMeshAdapter.cpp
+++ b/Engine/Plugins/Runtime/EditableMesh/Source/EditableMesh/EditableStaticMeshAdapter.cpp
@@ -20,7 +20,8 @@ UEditableStaticMeshAdapter::UEditableStaticMeshAdapter()
 	  StaticMeshLODIndex( 0 ),
 	  RecreateRenderStateContext(),
 	  CachedBoundingBoxAndSphere( FVector::ZeroVector, FVector::ZeroVector, 0.0f ),
-	  bUpdateCollisionNeeded( false )
+	  bUpdateCollisionNeeded( false ),
+	  bRecreateSimplifiedCollision( true )
 {
 }
 
@@ -1041,7 +1042,10 @@ void UEditableStaticMeshAdapter::UpdateCollision()
 	// @todo mesheditor collision: We're wiping the existing simplified collision and generating a simple bounding
 	// box collision, since that's the best we can do without impacting performance.  We always using visibility (complex)
 	// collision for traces while mesh editing (for hover/selection), so simplified collision isn't really important.
-	const bool bRecreateSimplifiedCollision = true;
+	if ( !bRecreateSimplifiedCollision )
+	{
+		return;
+	}
 
 	if( StaticMesh->BodySetup == nullptr )
 	{
@@ -1059,27 +1063,21 @@ void UEditableStaticMeshAdapter::UpdateCollision()
 	// NOTE: We don't bother calling Modify() on the BodySetup as EndModification() will rebuild this guy after every undo
 	// BodySetup->Modify();
 
-	if( bRecreateSimplifiedCollision )
+	if( BodySetup->AggGeom.GetElementCount() > 0 )
 	{
-		if( BodySetup->AggGeom.GetElementCount() > 0 )
-		{
-			BodySetup->RemoveSimpleCollision();
-		}
+		BodySetup->RemoveSimpleCollision();
 	}
 
 	BodySetup->InvalidatePhysicsData();
 
-	if( bRecreateSimplifiedCollision )
-	{
-		const FBoxSphereBounds Bounds = StaticMesh->GetBounds();
-
-		FKBoxElem BoxElem;
-		BoxElem.Center = Bounds.Origin;
-		BoxElem.X = Bounds.BoxExtent.X * 2.0f;
-		BoxElem.Y = Bounds.BoxExtent.Y * 2.0f;
-		BoxElem.Z = Bounds.BoxExtent.Z * 2.0f;
-		BodySetup->AggGeom.BoxElems.Add( BoxElem );
-	}
+	const FBoxSphereBounds Bounds = StaticMesh->GetBounds();
+
+	FKBoxElem BoxElem;
+	BoxElem.Center = Bounds.Origin;
+	BoxElem.X = Bounds.BoxExtent.X * 2.0f;
+	BoxElem.Y = Bounds.BoxExtent.Y * 2.0f;
+	BoxElem.Z = Bounds.BoxExtent.Z * 2.0f;
+	BodySetup->AggGeom.BoxElems.Add( BoxElem );
 
 	// Update all static mesh components that are using this mesh
 	// @todo mesheditor perf: This is a pretty heavy operation, and overlaps with what we're already doing in RecreateRenderStateContext
diff --git a/Engine/Plugins/Runtime/EditableMesh/Source/EditableMesh/EditableStaticMeshAdapter.h b/Engine/Plugins/Runtime/EditableMesh/Source/EditableMesh/EditableStaticMeshAdapter.h
--- a/Engine/Plugins/Runtime/EditableMesh/Source/EditableMesh/EditableStaticMeshAdapter.h
+++ b/Engine/Plugins/Runtime/EditableMesh/Source/EditableMesh/EditableStaticMeshAdapter.h
@@ -168,6 +168,8 @@ public:
 	virtual void GeometryHitTest(const FHitParamsIn& InParams, FHitParamsOut& OutParams) override;
 #endif // WITH_EDITOR
 
+	void SetRecreateSimpleCollision(bool bInRecreateSimplifiedCollision) { bRecreateSimplifiedCollision = bInRecreateSimplifiedCollision; }
+	bool GetRecreateSimpleCollision() const { return bRecreateSimplifiedCollision;  }
 
 private:
 
@@ -225,4 +227,7 @@ private:
 
 	/** Flag to indicate if a modification requires regenerating collision */
 	bool bUpdateCollisionNeeded;
+
+	/** Flag to indicate if the simple collision has to be regenerated */
+	bool bRecreateSimplifiedCollision;
 };
