Fixed crash after deleting a Static Mesh from the level that has been in Edit Mesh Mode

diff --git a/Engine/Plugins/Runtime/EditableMesh/Source/EditableMesh/EditableStaticMeshAdapter.cpp b/Engine/Plugins/Runtime/EditableMesh/Source/EditableMesh/EditableStaticMeshAdapter.cpp
--- a/Engine/Plugins/Runtime/EditableMesh/Source/EditableMesh/EditableStaticMeshAdapter.cpp
+++ b/Engine/Plugins/Runtime/EditableMesh/Source/EditableMesh/EditableStaticMeshAdapter.cpp
@@ -29,9 +29,12 @@ void UEditableStaticMeshAdapter::BeginDestroy()
 {
 	Super::BeginDestroy();
 
-	if (StaticMesh)
+	// Follow up on change made by CL #5110941
+	// StaticMesh is not controlled by an EditableMesh anymore
+	// Rebuild its resources if applicable
+	if (StaticMesh && !StaticMesh->IsPendingKillOrUnreachable())
 	{
-		StaticMesh->ReleaseResources();
+		StaticMesh->Build( true );
 	}
 }
 
