Fixed crash after deleting a Static Mesh from the level that has been in Edit Mesh Mode by reverting CL #5110941

diff --git a/Engine/Plugins/Runtime/EditableMesh/Source/EditableMesh/EditableStaticMeshAdapter.cpp b/Engine/Plugins/Runtime/EditableMesh/Source/EditableMesh/EditableStaticMeshAdapter.cpp
--- a/Engine/Plugins/Runtime/EditableMesh/Source/EditableMesh/EditableStaticMeshAdapter.cpp
+++ b/Engine/Plugins/Runtime/EditableMesh/Source/EditableMesh/EditableStaticMeshAdapter.cpp
@@ -25,20 +25,6 @@ UEditableStaticMeshAdapter::UEditableStaticMeshAdapter()
 {
 }
 
-void UEditableStaticMeshAdapter::BeginDestroy()
-{
-	Super::BeginDestroy();
-
-	// Follow up on change made by CL #5110941
-	// StaticMesh is not controlled by an EditableMesh anymore
-	// Rebuild its resources if applicable
-	if (StaticMesh && !StaticMesh->IsPendingKillOrUnreachable())
-	{
-		StaticMesh->Build( true );
-	}
-}
-
-
 
 inline void UEditableStaticMeshAdapter::EnsureIndexBufferIs32Bit()
 {
diff --git a/Engine/Plugins/Runtime/EditableMesh/Source/EditableMesh/EditableStaticMeshAdapter.h b/Engine/Plugins/Runtime/EditableMesh/Source/EditableMesh/EditableStaticMeshAdapter.h
--- a/Engine/Plugins/Runtime/EditableMesh/Source/EditableMesh/EditableStaticMeshAdapter.h
+++ b/Engine/Plugins/Runtime/EditableMesh/Source/EditableMesh/EditableStaticMeshAdapter.h
@@ -120,8 +120,6 @@ public:
 	/** Default constructor that initializes good defaults for UEditableStaticMeshAdapter */
 	UEditableStaticMeshAdapter();
 
-	virtual void BeginDestroy() override;
-
 	virtual void Serialize( FArchive& Ar ) override;
 
 	/** Creates a editable static mesh from the specified component and sub-mesh address */
