Refactor FChange into FSwapChange and new command-pattern-style FCommandChange

diff --git a/Engine/Plugins/Editor/MeshEditor/Source/MeshEditor/MeshEditorMode.h b/Engine/Plugins/Editor/MeshEditor/Source/MeshEditor/MeshEditorMode.h
--- a/Engine/Plugins/Editor/MeshEditor/Source/MeshEditor/MeshEditorMode.h
+++ b/Engine/Plugins/Editor/MeshEditor/Source/MeshEditor/MeshEditorMode.h
@@ -335,7 +335,7 @@ protected:
 	};
 
 
-	class FSelectOrDeselectMeshElementsChange : public FChange
+	class FSelectOrDeselectMeshElementsChange : public FSwapChange
 	{
 
 	public:
@@ -371,7 +371,7 @@ protected:
 	};
 
 
-	class FDeselectAllMeshElementsChange : public FChange
+	class FDeselectAllMeshElementsChange : public FSwapChange
 	{
 
 	public:
@@ -414,7 +414,7 @@ protected:
 	};
 
 
-	class FSetElementSelectionModeChange : public FChange
+	class FSetElementSelectionModeChange : public FSwapChange
 	{
 
 	public:
diff --git a/Engine/Plugins/Runtime/EditableMesh/Source/EditableMesh/EditableMesh.cpp b/Engine/Plugins/Runtime/EditableMesh/Source/EditableMesh/EditableMesh.cpp
--- a/Engine/Plugins/Runtime/EditableMesh/Source/EditableMesh/EditableMesh.cpp
+++ b/Engine/Plugins/Runtime/EditableMesh/Source/EditableMesh/EditableMesh.cpp
@@ -427,7 +427,7 @@ static void InvertRemapTable( TSparseArray<int32>& InvertedRemapTable, const TSp
 }
 
 
-class FCompactChange : public FChange
+class FCompactChange : public FSwapChange
 {
 public:
 
@@ -459,7 +459,7 @@ struct FUncompactChangeInput
 };
 
 
-class FUncompactChange : public FChange
+class FUncompactChange : public FSwapChange
 {
 public:
 
diff --git a/Engine/Plugins/Runtime/EditableMesh/Source/EditableMesh/EditableMeshChanges.h b/Engine/Plugins/Runtime/EditableMesh/Source/EditableMesh/EditableMeshChanges.h
--- a/Engine/Plugins/Runtime/EditableMesh/Source/EditableMesh/EditableMeshChanges.h
+++ b/Engine/Plugins/Runtime/EditableMesh/Source/EditableMesh/EditableMeshChanges.h
@@ -22,7 +22,7 @@ struct FDeleteOrphanVerticesChangeInput
 };
 
 
-class FDeleteOrphanVerticesChange : public FChange
+class FDeleteOrphanVerticesChange : public FSwapChange
 {
 
 public:
@@ -67,7 +67,7 @@ struct FDeleteVertexInstancesChangeInput
 };
 
 
-class FDeleteVertexInstancesChange : public FChange
+class FDeleteVertexInstancesChange : public FSwapChange
 {
 
 public:
@@ -111,7 +111,7 @@ struct FDeleteEdgesChangeInput
 };
 
 
-class FDeleteEdgesChange : public FChange
+class FDeleteEdgesChange : public FSwapChange
 {
 
 public:
@@ -153,7 +153,7 @@ struct FCreateVerticesChangeInput
 };
 
 
-class FCreateVerticesChange : public FChange
+class FCreateVerticesChange : public FSwapChange
 {
 
 public:
@@ -195,7 +195,7 @@ struct FCreateVertexInstancesChangeInput
 };
 
 
-class FCreateVertexInstancesChange : public FChange
+class FCreateVertexInstancesChange : public FSwapChange
 {
 
 public:
@@ -237,7 +237,7 @@ struct FCreateEdgesChangeInput
 };
 
 
-class FCreateEdgesChange : public FChange
+class FCreateEdgesChange : public FSwapChange
 {
 
 public:
@@ -279,7 +279,7 @@ struct FCreatePolygonsChangeInput
 };
 
 
-class FCreatePolygonsChange : public FChange
+class FCreatePolygonsChange : public FSwapChange
 {
 
 public:
@@ -337,7 +337,7 @@ struct FDeletePolygonsChangeInput
 };
 
 
-class FDeletePolygonsChange : public FChange
+class FDeletePolygonsChange : public FSwapChange
 {
 
 public:
@@ -376,7 +376,7 @@ struct FFlipPolygonsChangeInput
 	}
 };
 
-class FFlipPolygonsChange : public FChange
+class FFlipPolygonsChange : public FSwapChange
 {
 public:
 	FFlipPolygonsChange( const FFlipPolygonsChangeInput& InitInput )
@@ -400,7 +400,7 @@ struct FSetVerticesAttributesChangeInput
 };
 
 
-class FSetVerticesAttributesChange : public FChange
+class FSetVerticesAttributesChange : public FSwapChange
 {
 
 public:
@@ -435,7 +435,7 @@ struct FSetVertexInstancesAttributesChangeInput
 };
 
 
-class FSetVertexInstancesAttributesChange : public FChange
+class FSetVertexInstancesAttributesChange : public FSwapChange
 {
 
 public:
@@ -477,7 +477,7 @@ struct FSetEdgesAttributesChangeInput
 };
 
 
-class FSetEdgesAttributesChange : public FChange
+class FSetEdgesAttributesChange : public FSwapChange
 {
 
 public:
@@ -519,7 +519,7 @@ struct FSetPolygonsVertexAttributesChangeInput
 };
 
 
-class FSetPolygonsVertexAttributesChange : public FChange
+class FSetPolygonsVertexAttributesChange : public FSwapChange
 {
 
 public:
@@ -561,7 +561,7 @@ struct FChangePolygonsVertexInstancesChangeInput
 };
 
 
-class FChangePolygonsVertexInstancesChange : public FChange
+class FChangePolygonsVertexInstancesChange : public FSwapChange
 {
 
 public:
@@ -603,7 +603,7 @@ struct FSetEdgesVerticesChangeInput
 };
 
 
-class FSetEdgesVerticesChange : public FChange
+class FSetEdgesVerticesChange : public FSwapChange
 {
 
 public:
@@ -654,7 +654,7 @@ struct FInsertPolygonPerimeterVerticesChangeInput
 };
 
 
-class FInsertPolygonPerimeterVerticesChange : public FChange
+class FInsertPolygonPerimeterVerticesChange : public FSwapChange
 {
 
 public:
@@ -709,7 +709,7 @@ struct FRemovePolygonPerimeterVerticesChangeInput
 };
 
 
-class FRemovePolygonPerimeterVerticesChange : public FChange
+class FRemovePolygonPerimeterVerticesChange : public FSwapChange
 {
 
 public:
@@ -755,7 +755,7 @@ struct FStartOrEndModificationChangeInput
 };
 
 
-class FStartOrEndModificationChange : public FChange
+class FStartOrEndModificationChange : public FSwapChange
 {
 
 public:
@@ -792,7 +792,7 @@ struct FSetSubdivisionCountChangeInput
 };
 
 
-class FSetSubdivisionCountChange : public FChange
+class FSetSubdivisionCountChange : public FSwapChange
 {
 
 public:
@@ -829,7 +829,7 @@ struct FCreatePolygonGroupsChangeInput
 };
 
 
-class FCreatePolygonGroupsChange : public FChange
+class FCreatePolygonGroupsChange : public FSwapChange
 {
 
 public:
@@ -870,7 +870,7 @@ struct FDeletePolygonGroupsChangeInput
 };
 
 
-class FDeletePolygonGroupsChange : public FChange
+class FDeletePolygonGroupsChange : public FSwapChange
 {
 
 public:
@@ -915,7 +915,7 @@ struct FAssignPolygonsToPolygonGroupChangeInput
 };
 
 
-class FAssignPolygonsToPolygonGroupChange : public FChange
+class FAssignPolygonsToPolygonGroupChange : public FSwapChange
 {
 
 public:
diff --git a/Engine/Source/Editor/UnrealEd/Private/EditorTransaction.cpp b/Engine/Source/Editor/UnrealEd/Private/EditorTransaction.cpp
--- a/Engine/Source/Editor/UnrealEd/Private/EditorTransaction.cpp
+++ b/Engine/Source/Editor/UnrealEd/Private/EditorTransaction.cpp
@@ -233,8 +233,24 @@ void FTransaction::FObjectRecord::Load(FTransaction* Owner)
 
 		if (CustomChange.IsValid())
 		{
-			TUniquePtr<FChange> InvertedChange = CustomChange->Execute( Object.Get() );
-			CustomChange = MoveTemp( InvertedChange );
+			if (CustomChange->GetChangeType() == FChange::EChangeStyle::InPlaceSwap)
+			{
+				TUniquePtr<FChange> InvertedChange = CustomChange->Execute(Object.Get());
+				ensure(InvertedChange->GetChangeType() == FChange::EChangeStyle::InPlaceSwap);
+				CustomChange = MoveTemp(InvertedChange);
+			}
+			else
+			{
+				bool bIsRedo = (Owner->Inc == 1);
+				if (bIsRedo)
+				{
+					CustomChange->Apply(Object.Get());
+				}
+				else
+				{
+					CustomChange->Revert(Object.Get());
+				}
+			}
 		}
 		else
 		{
diff --git a/Engine/Source/Runtime/Core/Public/Misc/Change.h b/Engine/Source/Runtime/Core/Public/Misc/Change.h
--- a/Engine/Source/Runtime/Core/Public/Misc/Change.h
+++ b/Engine/Source/Runtime/Core/Public/Misc/Change.h
@@ -5,16 +5,42 @@
 #include "Templates/UniquePtr.h"
 #include "Containers/Array.h"
 
-// @todo mesheditor: Comment these classes and enums!
-
+/**
+ * FChange modifies a UObject and is meant to be used to implement undo/redo.
+ * The change is embedded in an FTransaction which executes it *instead* of the standard 
+ * serialization transaction (cannot be combined - see FTransaction).
+ *
+ * The original FChange style (used by MeshEditor) was that calling Execute() would return a new 
+ * FChange that applies the opposite action, and FTransaction would swap the two at each undo/redo
+ * step (eg a "DeleteObject" FChange would return a "CreateObject" FChange)
+ *
+ * The alternative "Command Pattern"-style FChange calls Apply() and Revert() on a single FChange.
+  *
+ * FChange may eventually be deprecated. You should subclass 
+ * FSwapChange and FCommandChange to implement these different styles.
+ */
 class CORE_API FChange
 {
 
 public:
+	enum class EChangeStyle
+	{
+		InPlaceSwap,			// Call Execute() which returns new "opposite" FChange   (default)
+		CommandPattern			// call Revert() to Undo and Apply() to Redo
+	};
+
+	/** What style of change is this */
+	virtual EChangeStyle GetChangeType() = 0;
 
 	/** Makes the change to the object, returning a new change that can be used to perfectly roll back this change */
 	virtual TUniquePtr<FChange> Execute( UObject* Object ) = 0;
 
+	/** Makes the change to the object */
+	virtual void Apply( UObject* Object ) = 0;
+
+	/** Reverts change to the object */
+	virtual void Revert( UObject* Object ) = 0;
+
 	/** Describes this change (for debugging) */
 	virtual FString ToString() const = 0;
 
@@ -42,6 +68,62 @@ private:
 
 
 
+/**
+ * To use FSwapChange you must implement Execute().
+ * This function must do two things:
+ *   1) apply the change to the given UObject
+ *   2) return a new FSwapChange that does the "opposite" action
+ */
+class CORE_API FSwapChange : public FChange
+{
+
+public:
+	virtual EChangeStyle GetChangeType()
+	{
+		return FChange::EChangeStyle::InPlaceSwap;
+	}
+
+	/** Makes the change to the object */
+	virtual void Apply(UObject* Object) override
+	{
+		check(false);
+	}
+
+	/** Reverts change to the object */
+	virtual void Revert(UObject* Object) override
+	{
+		check(false);
+	}
+
+};
+
+
+/**
+ * To use FCommandChange you must implement Apply() and Revert()
+ * Revert() is called to "Undo" and Apply() is called to "Redo"
+ */
+class CORE_API FCommandChange : public FChange
+{
+
+public:
+	virtual EChangeStyle GetChangeType()
+	{
+		return FChange::EChangeStyle::CommandPattern;
+	}
+
+	virtual TUniquePtr<FChange> Execute(UObject* Object)
+	{
+		check(false);
+		return nullptr;
+	}
+
+};
+
+
+
+
+
+
 struct CORE_API FCompoundChangeInput
 {
 	FCompoundChangeInput()
@@ -63,7 +145,12 @@ private:
 };
 
 
-class CORE_API FCompoundChange : public FChange
+
+/**
+ * FCompoundChange applies a sequence of FSwapChanges.
+ * The changes are executed in reverse order (this is like a mini undo stack)
+ */
+class CORE_API FCompoundChange : public FSwapChange
 {
 
 public:
