Remove some log spew from EditableMesh, very slow in modeling tools
#rnx

diff --git a/Engine/Plugins/Runtime/EditableMesh/Source/EditableMesh/EditableMesh.cpp b/Engine/Plugins/Runtime/EditableMesh/Source/EditableMesh/EditableMesh.cpp
--- a/Engine/Plugins/Runtime/EditableMesh/Source/EditableMesh/EditableMesh.cpp
+++ b/Engine/Plugins/Runtime/EditableMesh/Source/EditableMesh/EditableMesh.cpp
@@ -2621,7 +2621,8 @@ void UEditableMesh::SetSubdivisionCount( const int32 NewSubdivisionCount )
 
 void UEditableMesh::MoveVertices( const TArray<FVertexToMove>& VerticesToMove )
 {
-	EM_ENTER( TEXT( "MoveVertices: %s" ), *LogHelpers::ArrayToString( VerticesToMove ) );
+	//EM_ENTER( TEXT( "MoveVertices: %s" ), *LogHelpers::ArrayToString( VerticesToMove ) );
+	EM_ENTER(TEXT("MoveVertices: [redacted]"));
 
 	static TSet< FPolygonID > VertexConnectedPolygons;
 	VertexConnectedPolygons.Reset();
@@ -4415,7 +4416,8 @@ void UEditableMesh::DeletePolygonGroups( const TArray<FPolygonGroupID>& PolygonG
 
 void UEditableMesh::SetVerticesAttributes( const TArray<FAttributesForVertex>& AttributesForVertices )
 {
-	EM_ENTER( TEXT( "SetVerticesAttributes: %s" ), *LogHelpers::ArrayToString( AttributesForVertices ) );
+	//EM_ENTER( TEXT( "SetVerticesAttributes: %s" ), *LogHelpers::ArrayToString( AttributesForVertices ) );
+	EM_ENTER(TEXT("SetVerticesAttributes: [redacted]"));
 
 	FSetVerticesAttributesChangeInput RevertInput;
 
