MeshDescription now uses triangles as its primary face representation. They can and should be used in preference to polygons when a mesh is known to be entirely triangulated, as they are quicker, and 

diff --git a/Engine/Plugins/Runtime/EditableMesh/Source/EditableMesh/EditableMesh.cpp b/Engine/Plugins/Runtime/EditableMesh/Source/EditableMesh/EditableMesh.cpp
--- a/Engine/Plugins/Runtime/EditableMesh/Source/EditableMesh/EditableMesh.cpp
+++ b/Engine/Plugins/Runtime/EditableMesh/Source/EditableMesh/EditableMesh.cpp
@@ -993,13 +993,13 @@ FPolygonGroupID UEditableMesh::GetGroupForPolygon( const FPolygonID PolygonID )
 
 int32 UEditableMesh::GetPolygonPerimeterVertexCount( const FPolygonID PolygonID ) const
 {
-	return GetMeshDescription()->GetPolygonPerimeterVertexInstances( PolygonID ).Num();
+	return GetMeshDescription()->GetPolygonVertexInstances( PolygonID ).Num();
 }
 
 
 FVertexInstanceID UEditableMesh::GetPolygonPerimeterVertexInstance( const FPolygonID PolygonID, const int32 PolygonVertexNumber ) const
 {
-	const TArray<FVertexInstanceID>& VertexInstanceIDs = GetMeshDescription()->GetPolygonPerimeterVertexInstances( PolygonID );
+	const TArray<FVertexInstanceID>& VertexInstanceIDs = GetMeshDescription()->GetPolygonVertexInstances( PolygonID );
 	return ( PolygonVertexNumber >= 0 && PolygonVertexNumber < VertexInstanceIDs.Num() ) ? VertexInstanceIDs[ PolygonVertexNumber ] : FVertexInstanceID::Invalid;
 }
 
@@ -1013,14 +1013,14 @@ FVertexID UEditableMesh::GetPolygonPerimeterVertex( const FPolygonID PolygonID,
 
 int32 UEditableMesh::GetPolygonTriangulatedTriangleCount( const FPolygonID PolygonID ) const
 {
-	return GetMeshDescription()->GetPolygonTriangles( PolygonID ).Num();
+	return GetMeshDescription()->GetPolygonTriangleIDs( PolygonID ).Num();
 }
 
 
-FMeshTriangle UEditableMesh::GetPolygonTriangulatedTriangle( const FPolygonID PolygonID, int32 PolygonTriangleNumber ) const
+FTriangleID UEditableMesh::GetPolygonTriangulatedTriangle( const FPolygonID PolygonID, int32 PolygonTriangleNumber ) const
 {
-	const TArray<FMeshTriangle>& Triangles = GetMeshDescription()->GetPolygonTriangles( PolygonID );
-	return ( PolygonTriangleNumber >= 0 && PolygonTriangleNumber < Triangles.Num() ) ? Triangles[ PolygonTriangleNumber ] : FMeshTriangle();
+	const TArray<FTriangleID>& Triangles = GetMeshDescription()->GetPolygonTriangleIDs( PolygonID );
+	return ( PolygonTriangleNumber >= 0 && PolygonTriangleNumber < Triangles.Num() ) ? Triangles[ PolygonTriangleNumber ] : FTriangleID::Invalid;
 }
 
 
@@ -1218,7 +1218,6 @@ void UEditableMesh::ReplaceVertexInstanceInPolygons( const FVertexInstanceID Old
 	static TArray<FChangeVertexInstancesForPolygon>  VertexInstancesToChange;
 	VertexInstancesToChange.Reset( PolygonIDs.Num() );
 
-	const FPolygonArray& Polygons = GetMeshDescription()->Polygons();
 	for( const FPolygonID PolygonID : PolygonIDs )
 	{
 		VertexInstancesToChange.Emplace();
@@ -1226,8 +1225,7 @@ void UEditableMesh::ReplaceVertexInstanceInPolygons( const FVertexInstanceID Old
 
 		ChangeVertexInstances.PolygonID = PolygonID;
 
-		const FMeshPolygon& Polygon = Polygons[ PolygonID ];
-		int32 VertexInstanceIndex = Polygon.PerimeterContour.VertexInstanceIDs.Find( OldVertexInstanceID );
+		int32 VertexInstanceIndex = GetMeshDescription()->GetPolygonVertexInstances( PolygonID ).Find( OldVertexInstanceID );
 
 		if( VertexInstanceIndex != INDEX_NONE )
 		{
@@ -1561,20 +1559,20 @@ int32 UEditableMesh::GetPolygonPerimeterEdgeCount( const FPolygonID PolygonID )
 
 void UEditableMesh::GetPolygonPerimeterVertices( const FPolygonID PolygonID, TArray<FVertexID>& OutPolygonPerimeterVertexIDs ) const
 {
-	GetMeshDescription()->GetPolygonPerimeterVertices( PolygonID, OutPolygonPerimeterVertexIDs );
+	GetMeshDescription()->GetPolygonVertices( PolygonID, OutPolygonPerimeterVertexIDs );
 }
 
 
 void UEditableMesh::GetPolygonPerimeterVertexInstances( const FPolygonID PolygonID, TArray<FVertexInstanceID>& OutPolygonPerimeterVertexInstanceIDs ) const
 {
-	OutPolygonPerimeterVertexInstanceIDs = GetMeshDescription()->GetPolygonPerimeterVertexInstances( PolygonID );
+	OutPolygonPerimeterVertexInstanceIDs = GetMeshDescription()->GetPolygonVertexInstances( PolygonID );
 }
 
 
 FEdgeID UEditableMesh::GetPolygonPerimeterEdge( const FPolygonID PolygonID, const int32 PerimeterEdgeNumber, bool& bOutEdgeWindingIsReversedForPolygon ) const
 {
 	const FMeshDescription* Description = GetMeshDescription();
-	const TArray<FVertexInstanceID>& VertexInstanceIDs = Description->GetPolygonPerimeterVertexInstances( PolygonID );
+	const TArray<FVertexInstanceID>& VertexInstanceIDs = Description->GetPolygonVertexInstances( PolygonID );
 	if( PerimeterEdgeNumber >= VertexInstanceIDs.Num() )
 	{
 		return FEdgeID::Invalid;
@@ -1602,7 +1600,7 @@ FEdgeID UEditableMesh::GetVertexPairEdge( const FVertexID StartVertexID, const F
 void UEditableMesh::GetPolygonPerimeterEdges( const FPolygonID PolygonID, TArray<FEdgeID>& OutPolygonPerimeterEdgeIDs ) const
 {
 	const FMeshDescription* Description = GetMeshDescription();
-	const TArray<FVertexInstanceID>& VertexInstanceIDs = Description->GetPolygonPerimeterVertexInstances( PolygonID );
+	const TArray<FVertexInstanceID>& VertexInstanceIDs = Description->GetPolygonVertexInstances( PolygonID );
 	const int32 NumContourEdges = VertexInstanceIDs.Num();
 
 	OutPolygonPerimeterEdgeIDs.SetNumUninitialized( NumContourEdges, false );
@@ -2023,7 +2021,7 @@ void UEditableMesh::GenerateOpenSubdivLimitSurfaceData()
 
 					for( const FPolygonID PolygonID : PolygonIDs )
 					{
-						for( const FVertexInstanceID VertexInstanceID : GetMeshDescription()->GetPolygonPerimeterVertexInstances( PolygonID ) )
+						for( const FVertexInstanceID VertexInstanceID : GetMeshDescription()->GetPolygonVertexInstances( PolygonID ) )
 						{
 							FOsdFVarVertexData& FVarVertexData = *new( FVarVertexDatas ) FOsdFVarVertexData();
 							FVarVertexData.TextureCoordinates[ 0 ] = TextureCoordinateCount > 0 ? VertexUVs.Get( VertexInstanceID, 0 ) : FVector2D::ZeroVector;
@@ -2398,8 +2396,7 @@ void UEditableMesh::RetriangulatePolygons()
 	// Perform triangulation directly into mesh polygons
 	for( const FPolygonID PolygonID : PolygonsPendingTriangulation )
 	{
-		FMeshPolygon& Polygon = Polygons[ PolygonID ];
-		ComputePolygonTriangulation( PolygonID, Polygon.Triangles );
+		GetMeshDescription()->ComputePolygonTriangulation( PolygonID );
 	}
 
 	for( UEditableMeshAdapter* Adapter : Adapters )
@@ -2409,170 +2406,18 @@ void UEditableMesh::RetriangulatePolygons()
 }
 
 
-void UEditableMesh::ComputePolygonTriangulation( const FPolygonID PolygonID, TArray<FMeshTriangle>& OutTriangles ) const
-{
-	// NOTE: This polygon triangulation code is partially based on the ear cutting algorithm described on
-	//       page 497 of the book "Real-time Collision Detection", published in 2005.
-
-	struct Local
-	{
-		// Returns true if the triangle formed by the specified three positions has a normal that is facing the opposite direction of the reference normal
-		static inline bool IsTriangleFlipped( const FVector ReferenceNormal, const FVector VertexPositionA, const FVector VertexPositionB, const FVector VertexPositionC )
-		{
-			const FVector TriangleNormal = FVector::CrossProduct(
-				VertexPositionC - VertexPositionA,
-				VertexPositionB - VertexPositionA ).GetSafeNormal();
-			return ( FVector::DotProduct( ReferenceNormal, TriangleNormal ) <= 0.0f );
-		}
-
-	};
-
-
-	OutTriangles.Reset();
-
-	// @todo mesheditor: Perhaps should always attempt to triangulate by splitting polygons along the shortest edge, for better determinism.
-
-//	const FMeshPolygon& Polygon = GetMeshDescription()->GetPolygon( PolygonID );
-	const TArray<FVertexInstanceID>& PolygonVertexInstanceIDs = GetMeshDescription()->GetPolygonPerimeterVertexInstances( PolygonID );
-
-	// Polygon must have at least three vertices/edges
-	const int32 PolygonVertexCount = PolygonVertexInstanceIDs.Num();
-	check( PolygonVertexCount >= 3 );
-
-	// First figure out the polygon normal.  We need this to determine which triangles are convex, so that
-	// we can figure out which ears to clip
-	const FVector PolygonNormal = ComputePolygonNormal( PolygonID );
-
-	// Make a simple linked list array of the previous and next vertex numbers, for each vertex number
-	// in the polygon.  This will just save us having to iterate later on.
-	static TArray<int32> PrevVertexNumbers, NextVertexNumbers;
-	static TArray<FVector> VertexPositions;
-
-	{
-		TVertexAttributesConstRef<FVector> MeshVertexPositions = GetMeshDescription()->VertexAttributes().GetAttributesRef<FVector>( MeshAttribute::Vertex::Position );
-		PrevVertexNumbers.SetNumUninitialized( PolygonVertexCount, false );
-		NextVertexNumbers.SetNumUninitialized( PolygonVertexCount, false );
-		VertexPositions.SetNumUninitialized( PolygonVertexCount, false );
-
-		for( int32 VertexNumber = 0; VertexNumber < PolygonVertexCount; ++VertexNumber )
-		{
-			PrevVertexNumbers[ VertexNumber ] = VertexNumber - 1;
-			NextVertexNumbers[ VertexNumber ] = VertexNumber + 1;
-
-			VertexPositions[ VertexNumber ] = MeshVertexPositions[ GetMeshDescription()->GetVertexInstanceVertex( PolygonVertexInstanceIDs[ VertexNumber ] ) ];
-		}
-		PrevVertexNumbers[ 0 ] = PolygonVertexCount - 1;
-		NextVertexNumbers[ PolygonVertexCount - 1 ] = 0;
-	}
-
-	int32 EarVertexNumber = 0;
-	int32 EarTestCount = 0;
-	for( int32 RemainingVertexCount = PolygonVertexCount; RemainingVertexCount >= 3; )
-	{
-		bool bIsEar = true;
-
-		// If we're down to only a triangle, just treat it as an ear.  Also, if we've tried every possible candidate
-		// vertex looking for an ear, go ahead and just treat the current vertex as an ear.  This can happen when 
-		// vertices are colinear or other degenerate cases.
-		if( RemainingVertexCount > 3 && EarTestCount < RemainingVertexCount )
-		{
-			const FVector PrevVertexPosition = VertexPositions[ PrevVertexNumbers[ EarVertexNumber ] ];
-			const FVector EarVertexPosition = VertexPositions[ EarVertexNumber ];
-			const FVector NextVertexPosition = VertexPositions[ NextVertexNumbers[ EarVertexNumber ] ];
-
-			// Figure out whether the potential ear triangle is facing the same direction as the polygon
-			// itself.  If it's facing the opposite direction, then we're dealing with a concave triangle
-			// and we'll skip it for now.
-			if( !Local::IsTriangleFlipped(
-					PolygonNormal,
-					PrevVertexPosition,
-					EarVertexPosition,
-					NextVertexPosition ) )
-			{
-				int32 TestVertexNumber = NextVertexNumbers[ NextVertexNumbers[ EarVertexNumber ] ];
-
-				do 
-				{
-					// Test every other remaining vertex to make sure that it doesn't lie inside our potential ear
-					// triangle.  If we find a vertex that's inside the triangle, then it cannot actually be an ear.
-					const FVector TestVertexPosition = VertexPositions[ TestVertexNumber ];
-					if( FGeomTools::PointInTriangle(
-							PrevVertexPosition,
-							EarVertexPosition,
-							NextVertexPosition,
-							TestVertexPosition,
-							SMALL_NUMBER ) )
-					{
-						bIsEar = false;
-						break;
-					}
-
-					TestVertexNumber = NextVertexNumbers[ TestVertexNumber ];
-				} 
-				while( TestVertexNumber != PrevVertexNumbers[ EarVertexNumber ] );
-			}
-			else
-			{
-				bIsEar = false;
-			}
-		}
-
-		if( bIsEar )
-		{
-			// OK, we found an ear!  Let's save this triangle in our output buffer.
-			{
-				OutTriangles.Emplace();
-				FMeshTriangle& Triangle = OutTriangles.Last();
-
-				Triangle.SetVertexInstanceID( 0, PolygonVertexInstanceIDs[ PrevVertexNumbers[ EarVertexNumber ] ] );
-				Triangle.SetVertexInstanceID( 1, PolygonVertexInstanceIDs[ EarVertexNumber ] );
-				Triangle.SetVertexInstanceID( 2, PolygonVertexInstanceIDs[ NextVertexNumbers[ EarVertexNumber ] ] );
-			}
-
-			// Update our linked list.  We're effectively cutting off the ear by pointing the ear vertex's neighbors to
-			// point at their next sequential neighbor, and reducing the remaining vertex count by one.
-			{
-				NextVertexNumbers[ PrevVertexNumbers[ EarVertexNumber ] ] = NextVertexNumbers[ EarVertexNumber ];
-				PrevVertexNumbers[ NextVertexNumbers[ EarVertexNumber ] ] = PrevVertexNumbers[ EarVertexNumber ];
-				--RemainingVertexCount;
-			}
-
-			// Move on to the previous vertex in the list, now that this vertex was cut
-			EarVertexNumber = PrevVertexNumbers[ EarVertexNumber ];
-
-			EarTestCount = 0;
-		}
-		else
-		{
-			// The vertex is not the ear vertex, because it formed a triangle that either had a normal which pointed in the opposite direction
-			// of the polygon, or at least one of the other polygon vertices was found to be inside the triangle.  Move on to the next vertex.
-			EarVertexNumber = NextVertexNumbers[ EarVertexNumber ];
-
-			// Keep track of how many ear vertices we've tested, so that if we exhaust all remaining vertices, we can
-			// fall back to clipping the triangle and adding it to our mesh anyway.  This is important for degenerate cases.
-			++EarTestCount;
-		}
-	}
-
-	check( OutTriangles.Num() > 0 );
-}
-
-
-bool UEditableMesh::ComputeBarycentricWeightForPointOnPolygon( const FPolygonID PolygonID, const FVector PointOnPolygon, FMeshTriangle& OutTriangle, FVector& OutTriangleVertexWeights ) const
+FTriangleID UEditableMesh::ComputeBarycentricWeightForPointOnPolygon( const FPolygonID PolygonID, const FVector PointOnPolygon, FVector& OutTriangleVertexWeights ) const
 {
 	// @todo mesheditor: Modify this method so it can return a meaningful barycentric weight for an off-polygon position, referencing the closest triangle, with out-of-range weights.
 
-	const FMeshPolygon& Polygon = GetMeshDescription()->GetPolygon( PolygonID );
-	const int32 TriangleCount = Polygon.Triangles.Num();
-
 	TVertexAttributesConstRef<FVector> VertexPositions = GetMeshDescription()->VertexAttributes().GetAttributesRef<FVector>( MeshAttribute::Vertex::Position );
 
 	// Figure out which triangle the incoming point is within
-	for( const FMeshTriangle& Triangle : Polygon.Triangles )
+	for( const FTriangleID TriangleID : GetMeshDescription()->GetPolygonTriangleIDs( PolygonID ) )
 	{
-		const FVector TriangleVertex0Position = VertexPositions[ GetMeshDescription()->GetVertexInstanceVertex( Triangle.VertexInstanceID0 ) ];
-		const FVector TriangleVertex1Position = VertexPositions[ GetMeshDescription()->GetVertexInstanceVertex( Triangle.VertexInstanceID1 ) ];
-		const FVector TriangleVertex2Position = VertexPositions[ GetMeshDescription()->GetVertexInstanceVertex( Triangle.VertexInstanceID2 ) ];
+		const FVector TriangleVertex0Position = VertexPositions[ GetMeshDescription()->GetVertexInstanceVertex( GetMeshDescription()->GetTriangleVertexInstance( TriangleID, 0 ) ) ];
+		const FVector TriangleVertex1Position = VertexPositions[ GetMeshDescription()->GetVertexInstanceVertex( GetMeshDescription()->GetTriangleVertexInstance( TriangleID, 1 ) ) ];
+		const FVector TriangleVertex2Position = VertexPositions[ GetMeshDescription()->GetVertexInstanceVertex( GetMeshDescription()->GetTriangleVertexInstance( TriangleID, 2 ) ) ];
 
 		// Calculate the barycentric weights for the triangle's verts and determine if the point lies within its bounds.
 		OutTriangleVertexWeights = FMath::ComputeBaryCentric2D( PointOnPolygon, TriangleVertex0Position, TriangleVertex1Position, TriangleVertex2Position );
@@ -2580,12 +2425,11 @@ bool UEditableMesh::ComputeBarycentricWeightForPointOnPolygon( const FPolygonID
 		if( OutTriangleVertexWeights.X >= 0.0f && OutTriangleVertexWeights.Y >= 0.0f && OutTriangleVertexWeights.Z >= 0.0f )
 		{
 			// Okay, we found a triangle that the point is inside!  Return the corresponding vertex instances.
-			OutTriangle = Triangle;
-			return true;
+			return TriangleID;
 		}
 	}
 
-	return false;
+	return FTriangleID::Invalid;
 }
 
 
@@ -2805,7 +2649,7 @@ void UEditableMesh::SplitEdge( const FEdgeID EdgeID, const TArray<float>& Splits
 		// Iterate through each polygon connected to the edge being split
 		for( const FPolygonID ConnectedPolygonID : GetMeshDescription()->GetEdgeConnectedPolygons( EdgeID ) )
 		{
-			const TArray<FVertexInstanceID>& PerimeterVertexInstances = GetMeshDescription()->GetPolygonPerimeterVertexInstances( ConnectedPolygonID );
+			const TArray<FVertexInstanceID>& PerimeterVertexInstances = GetMeshDescription()->GetPolygonVertexInstances( ConnectedPolygonID );
 			const int32 NumVertexInstances = PerimeterVertexInstances.Num();
 
 			// Find vertex indices around the perimeter corresponding to the vertices defining the edge
@@ -2869,7 +2713,7 @@ void UEditableMesh::SplitEdge( const FEdgeID EdgeID, const TArray<float>& Splits
 			const bool bReversed = ( ( VertexIndices[ 0 ] > VertexIndices[ 1 ] ) ^ ( FMath::Abs( VertexIndices[ 0 ] - VertexIndices[ 1 ] ) > 1 ) );
 
 			FAffectedPolygonEdge AffectedPolygonEdge;
-			AffectedPolygonEdge.PerimeterVertexIndex = ( bReversed ? VertexIndices[ 1 ] : VertexIndices[ 0 ] ) + 1;
+			AffectedPolygonEdge.PerimeterVertexIndex = ( ( bReversed ? VertexIndices[ 1 ] : VertexIndices[ 0 ] ) + 1 ) % NumVertexInstances;
 			AffectedPolygonEdge.NewVertexInstanceArrayIndex = NewVertexInstanceArrayIndex;
 			AffectedPolygonEdge.bReversed = bReversed;
 
@@ -2879,30 +2723,92 @@ void UEditableMesh::SplitEdge( const FEdgeID EdgeID, const TArray<float>& Splits
 		CreateVertexInstances( VertexInstancesToCreate, NewVertexInstanceIDs );
 	}
 
-	{
-		static TArray<FVerticesForEdge> VerticesForEdges;
-		VerticesForEdges.Reset();
+	// Get list of polygons to remove
+	TArray<FPolygonID> AffectedPolygons = GetMeshDescription()->GetEdgeConnectedPolygons( EdgeID );
+
+	TArray<FPolygonToCreate> PolygonsToCreate;
+	PolygonsToCreate.Reset( AffectedPolygons.Num() );
 
-		// We'll keep the existing edge, but update it to connect to the first new vertex.  The second vertex of the edge will 
-		// now connect to the first (new) vertex ID, and so on.  The incoming vertices are expected to be ordered correctly.
+	// Update all affected polygons with their new vertices.  Also, we'll fill in polygon-specific vertex attributes (texture coordinates)
+	{
+		for( const FPolygonID ConnectedPolygonID : GetMeshDescription()->GetEdgeConnectedPolygons( EdgeID ) )
 		{
-			FVerticesForEdge& VerticesForEdge = *new( VerticesForEdges ) FVerticesForEdge();
+			const FPolygonGroupID PolygonGroupID = GetMeshDescription()->GetPolygonPolygonGroup( ConnectedPolygonID );
+			TArray<FVertexInstanceID> PerimeterVertexInstanceIDs = GetMeshDescription()->GetPolygonVertexInstances( ConnectedPolygonID );
+
+			FPolygonToCreate& PolygonToCreate = PolygonsToCreate.Emplace_GetRef();
+			PolygonToCreate.PolygonGroupID = PolygonGroupID;
+			PolygonToCreate.PolygonEdgeHardness = EPolygonEdgeHardness::NewEdgesSoft;
+			PolygonToCreate.PerimeterVertices.Reserve( PerimeterVertexInstanceIDs.Num() + NumNewVertices );
+
+			const FAffectedPolygonEdge& AffectedPolygonEdge = PolygonToAffectedEdge.FindRef( ConnectedPolygonID );
+
+			for( int32 VertexNumber = 0; VertexNumber < PerimeterVertexInstanceIDs.Num(); ++VertexNumber )
+			{
+				if( VertexNumber == AffectedPolygonEdge.PerimeterVertexIndex )
+				{
+					for( int32 Count = 0; Count < NumNewVertices; ++Count )
+					{
+						const int32 VertexIndex = AffectedPolygonEdge.bReversed ? NumNewVertices - 1 - Count : Count;
+						FVertexAndAttributes& VertexAndAttributes = PolygonToCreate.PerimeterVertices.Emplace_GetRef();
+						VertexAndAttributes.VertexInstanceID = NewVertexInstanceIDs[ AffectedPolygonEdge.NewVertexInstanceArrayIndex + VertexIndex ];
+					}
+				}
 
-			VerticesForEdge.EdgeID = EdgeID;
-			VerticesForEdge.NewVertexID0 = OriginalEdgeVertexIDs[ 0 ];
-			VerticesForEdge.NewVertexID1 = NewVertexIDs[ 0 ];
+				FVertexAndAttributes& VertexAndAttributes = PolygonToCreate.PerimeterVertices.Emplace_GetRef();
+				VertexAndAttributes.VertexInstanceID = PerimeterVertexInstanceIDs[ VertexNumber ];
+			}
+
+			// Add the new vertices to the polygon
+//			InsertPolygonPerimeterVertices( ConnectedPolygonID, AffectedPolygonEdge.PerimeterVertexIndex, VerticesToInsert );
 		}
+	}
 
-		SetEdgesVertices( VerticesForEdges );
+	// Delete the old polygons - this will also delete the old edge
+	{
+		// Delete in opposite order so that the newly created ones get the same IDs as the corresponding deleted ones
+		TArray<FPolygonID> PolygonsToDelete;
+		PolygonsToDelete = AffectedPolygons;
+		Algo::Reverse( PolygonsToDelete );
+
+		const bool bDeleteOrphanEdges = true;
+		const bool bDeleteOrphanVertices = false;
+		const bool bDeleteOrphanVertexInstances = false;
+		const bool bDeleteEmptyPolygonGroups = false;
+		DeletePolygons( PolygonsToDelete, bDeleteOrphanEdges, bDeleteOrphanVertices, bDeleteOrphanVertexInstances, bDeleteEmptyPolygonGroups );
 	}
 
+	//{
+	//	static TArray<FVerticesForEdge> VerticesForEdges;
+	//	VerticesForEdges.Reset();
+
+	//	// We'll keep the existing edge, but update it to connect to the first new vertex.  The second vertex of the edge will 
+	//	// now connect to the first (new) vertex ID, and so on.  The incoming vertices are expected to be ordered correctly.
+	//	{
+	//		FVerticesForEdge& VerticesForEdge = *new( VerticesForEdges ) FVerticesForEdge();
+
+	//		VerticesForEdge.EdgeID = EdgeID;
+	//		VerticesForEdge.NewVertexID0 = OriginalEdgeVertexIDs[ 0 ];
+	//		VerticesForEdge.NewVertexID1 = NewVertexIDs[ 0 ];
+	//	}
+
+	//	SetEdgesVertices( VerticesForEdges );
+	//}
+
 	// Create new edges.  One for each of the new vertex positions passed in.
 	{
 		const int32 NewEdgeCount = NumNewVertices;
 
 		static TArray<FEdgeToCreate> EdgesToCreate;
 		EdgesToCreate.Reset();
-		EdgesToCreate.Reserve( NewEdgeCount );
+		EdgesToCreate.Reserve( NewEdgeCount + 1 );
+
+		FEdgeToCreate& FirstEdgeToCreate = *new( EdgesToCreate ) FEdgeToCreate();
+
+		FirstEdgeToCreate.VertexID0 = OriginalEdgeVertexIDs[ 0 ];
+		FirstEdgeToCreate.VertexID1 = NewVertexIDs[ 0 ];
+		FirstEdgeToCreate.EdgeAttributes = EdgeAttributeList;
+
 		for( int32 NewEdgeNumber = 0; NewEdgeNumber < NewEdgeCount; ++NewEdgeNumber )
 		{
 			FEdgeToCreate& EdgeToCreate = *new( EdgesToCreate ) FEdgeToCreate();
@@ -2910,45 +2816,29 @@ void UEditableMesh::SplitEdge( const FEdgeID EdgeID, const TArray<float>& Splits
 			EdgeToCreate.VertexID0 = NewVertexIDs[ NewEdgeNumber ];
 			EdgeToCreate.VertexID1 = ( NewEdgeNumber == ( NewEdgeCount - 1 ) ) ? OriginalEdgeFarVertexID : NewVertexIDs[ NewEdgeNumber + 1 ];
 
-			EdgeToCreate.ConnectedPolygons = GetMeshDescription()->GetEdgeConnectedPolygons( EdgeID );
+			//			EdgeToCreate.ConnectedPolygons = GetMeshDescription()->GetEdgeConnectedPolygons( EdgeID );
 
 			// Copy edge attributes over from original edge
 			EdgeToCreate.EdgeAttributes = EdgeAttributeList;
 		}
 
-		static TArray< FEdgeID > NewEdgeIDs;
-		NewEdgeIDs.Reset();
+		TArray< FEdgeID > NewEdgeIDs;
 		CreateEdges( EdgesToCreate, /* Out */ NewEdgeIDs );
 	}
 
-	// Update all affected polygons with their new vertices.  Also, we'll fill in polygon-specific vertex attributes (texture coordinates)
-	{
-		for( const FPolygonID ConnectedPolygonID : GetMeshDescription()->GetEdgeConnectedPolygons( EdgeID ) )
-		{
-			static TArray<FVertexAndAttributes> VerticesToInsert;
-			VerticesToInsert.Reset( NumNewVertices );
 
-			const FAffectedPolygonEdge& AffectedPolygonEdge = PolygonToAffectedEdge.FindRef( ConnectedPolygonID );
-
-			for( int32 Count = 0; Count < NumNewVertices; ++Count )
-			{
-				const int32 VertexIndex = AffectedPolygonEdge.bReversed ? NumNewVertices - 1 - Count : Count;
-
-				VerticesToInsert.Emplace();
-				FVertexAndAttributes& VertexAndAttributes = VerticesToInsert.Last();
-				VertexAndAttributes.VertexInstanceID = NewVertexInstanceIDs[ AffectedPolygonEdge.NewVertexInstanceArrayIndex + VertexIndex ];
-			}
-
-			// Add the new vertices to the polygon
-			InsertPolygonPerimeterVertices( ConnectedPolygonID, AffectedPolygonEdge.PerimeterVertexIndex, VerticesToInsert );
-		}
+	// Create the new polygons
+	{
+		TArray<FEdgeID> NewEdgeIDs;
+		TArray<FPolygonID> NewPolygonIDs;
+		CreatePolygons( PolygonsToCreate, /* Out */ NewPolygonIDs, /* Out */ NewEdgeIDs );
 	}
 
 	// Generate normals and tangents
-	PolygonsPendingNewTangentBasis.Append( GetMeshDescription()->GetEdgeConnectedPolygons( EdgeID ) );
+	PolygonsPendingNewTangentBasis.Append( AffectedPolygons );
 
 	// Retriangulate all of the affected polygons
-	PolygonsPendingTriangulation.Append( GetMeshDescription()->GetEdgeConnectedPolygons( EdgeID ) );
+	PolygonsPendingTriangulation.Append( AffectedPolygons );
 
 	EM_EXIT( TEXT( "SplitEdge returned %s" ), *LogHelpers::ArrayToString( OutNewVertexIDs ) );
 }
@@ -3589,8 +3479,6 @@ void UEditableMesh::DeleteOrphanVertices( const TArray<FVertexID>& VertexIDsToDe
 {
 	EM_ENTER( TEXT( "DeleteOrphanVertices: %s" ), *LogHelpers::ArrayToString( VertexIDsToDelete ) );
 
-	FVertexArray& Vertices = GetMeshDescription()->Vertices();
-
 	// Back everything up
 	{
 		FCreateVerticesChangeInput RevertInput;
@@ -3603,9 +3491,8 @@ void UEditableMesh::DeleteOrphanVertices( const TArray<FVertexID>& VertexIDsToDe
 
 			// Make sure the vertex is truly an orphan.  We're not going to be able to restore its polygon vertex attributes,
 			// because the polygons won't exist when we're restoring the change
-			const FMeshVertex& Vertex = Vertices[ VertexID ];
-			check( Vertex.ConnectedEdgeIDs.Num() == 0 );
-			check( Vertex.VertexInstanceIDs.Num() == 0 );
+			check( GetMeshDescription()->GetVertexConnectedEdges( VertexID ).Num() == 0 );
+			check( GetMeshDescription()->GetVertexVertexInstances( VertexID ).Num() == 0 );
 
 			RevertInput.VerticesToCreate.Emplace();
 			FVertexToCreate& VertexToCreate = RevertInput.VerticesToCreate.Last();
@@ -3697,9 +3584,6 @@ void UEditableMesh::DeleteEdges( const TArray<FEdgeID>& EdgeIDsToDelete, const b
 {
 	EM_ENTER( TEXT( "DeleteEdges: %s %s" ), *LogHelpers::ArrayToString( EdgeIDsToDelete ), *LogHelpers::BoolToString( bDeleteOrphanedVertices ) );
 
-	FEdgeArray& Edges = GetMeshDescription()->Edges();
-	FVertexArray& Vertices = GetMeshDescription()->Vertices();
-
 	// Back everything up
 	{
 		FCreateEdgesChangeInput RevertInput;
@@ -3708,15 +3592,14 @@ void UEditableMesh::DeleteEdges( const TArray<FEdgeID>& EdgeIDsToDelete, const b
 		for( int32 EdgeNumber = EdgeIDsToDelete.Num() - 1; EdgeNumber >= 0; --EdgeNumber )
 		{
 			const FEdgeID EdgeID = EdgeIDsToDelete[ EdgeNumber ];
-			const FMeshEdge& Edge = Edges[ EdgeID ];
 
 			RevertInput.EdgesToCreate.Emplace();
 			FEdgeToCreate& EdgeToCreate = RevertInput.EdgesToCreate.Last();
 
 			EdgeToCreate.OriginalEdgeID = EdgeID;
-			EdgeToCreate.VertexID0 = Edge.VertexIDs[ 0 ];
-			EdgeToCreate.VertexID1 = Edge.VertexIDs[ 1 ];
-			EdgeToCreate.ConnectedPolygons = Edge.ConnectedPolygons;
+			EdgeToCreate.VertexID0 = GetMeshDescription()->GetEdgeVertex( EdgeID, 0 );
+			EdgeToCreate.VertexID1 = GetMeshDescription()->GetEdgeVertex( EdgeID, 1 );
+//			EdgeToCreate.ConnectedPolygons = GetMeshDescription()->GetEdgeConnectedPolygons( EdgeID );
 			
 			BackupAllAttributes( EdgeToCreate.EdgeAttributes, GetMeshDescription()->EdgeAttributes(), EdgeID );
 		}
@@ -3915,11 +3798,11 @@ void UEditableMesh::CreateEdges( const TArray<FEdgeToCreate>& EdgesToCreate, TAr
 			FEdgeID EdgeID = EdgeToCreate.OriginalEdgeID;
 			if( EdgeID != FEdgeID::Invalid )
 			{
-				GetMeshDescription()->CreateEdgeWithID( EdgeID, EdgeToCreate.VertexID0, EdgeToCreate.VertexID1, EdgeToCreate.ConnectedPolygons );
+				GetMeshDescription()->CreateEdgeWithID( EdgeID, EdgeToCreate.VertexID0, EdgeToCreate.VertexID1 );
 			}
 			else
 			{
-				EdgeID = GetMeshDescription()->CreateEdge( EdgeToCreate.VertexID0, EdgeToCreate.VertexID1, EdgeToCreate.ConnectedPolygons );
+				EdgeID = GetMeshDescription()->CreateEdge( EdgeToCreate.VertexID0, EdgeToCreate.VertexID1 );
 			}
 
 			OutNewEdgeIDs.Add( EdgeID );
@@ -3993,12 +3876,59 @@ FVertexInstanceID UEditableMesh::CreateVertexInstanceForContourVertex( const FVe
 }
 
 
-void UEditableMesh::CreatePolygonContour( const TArray<FVertexAndAttributes>& Contour, TArray<FVertexInstanceID>& OutVertexInstanceIDs )
+void UEditableMesh::CreatePolygonContour( const TArray<FVertexAndAttributes>& Contour, TArray<FEdgeID>& OutEdgeIDs, TArray<FVertexInstanceID>& OutVertexInstanceIDs )
 {
 	// All polygons must have at least three vertices
 	const int32 NumContourVertices = Contour.Num();
 	check( NumContourVertices >= 3 );
 
+	OutEdgeIDs.Reset( NumContourVertices );
+
+	static TArray<FEdgeToCreate> EdgesToCreate;
+	EdgesToCreate.Reset();
+
+	// Get edges or decide to create them
+	for( int32 VertexNumber = 0; VertexNumber < NumContourVertices; ++VertexNumber )
+	{
+		const int32 NextVertexNumber = ( VertexNumber + 1 ) % NumContourVertices;
+
+		// Get vertex ID for the ends of the edge.
+		FVertexID VertexID0 = Contour[ VertexNumber ].VertexID;
+		FVertexID VertexID1 = Contour[ NextVertexNumber ].VertexID;
+
+		// If valid vertex IDs weren't supplied, we'll deduce it from the vertex instance IDs instead (which must be valid instead).
+		if( VertexID0 == FVertexID::Invalid )
+		{
+			const FVertexInstanceID VertexInstanceID0 = Contour[ VertexNumber ].VertexInstanceID;
+			check( VertexInstanceID0 != FVertexInstanceID::Invalid );
+			VertexID0 = GetVertexInstanceVertex( VertexInstanceID0 );
+		}
+
+		if( VertexID1 == FVertexID::Invalid )
+		{
+			const FVertexInstanceID VertexInstanceID1 = Contour[ NextVertexNumber ].VertexInstanceID;
+			check( VertexInstanceID1 != FVertexInstanceID::Invalid );
+			VertexID1 = GetVertexInstanceVertex( VertexInstanceID1 );
+		}
+
+		FEdgeID EdgeID = GetMeshDescription()->GetVertexPairEdge( VertexID0, VertexID1 );
+
+		if( EdgeID == FEdgeID::Invalid )
+		{
+			// Create the new edge. Don't specify any attributes now.
+			EdgesToCreate.Emplace();
+			FEdgeToCreate& EdgeToCreate = EdgesToCreate.Last();
+			EdgeToCreate.VertexID0 = VertexID0;
+			EdgeToCreate.VertexID1 = VertexID1;
+		}
+	}
+
+	// Create missing edges
+	if( EdgesToCreate.Num() > 0 )
+	{
+		CreateEdges( EdgesToCreate, OutEdgeIDs );
+	}
+
 	OutVertexInstanceIDs.SetNumUninitialized( NumContourVertices );
 
 	static TArray<FVertexInstanceToCreate> VertexInstancesToCreate;
@@ -4079,17 +4009,19 @@ void UEditableMesh::CreatePolygons( const TArray<FPolygonToCreate>& PolygonsToCr
 			static TArray<FVertexInstanceID> PerimeterVertexInstances;
 			static TArray<FEdgeID> NewEdgeIDsForContour;
 
-			CreatePolygonContour( PolygonToCreate.PerimeterVertices, PerimeterVertexInstances );
+			CreatePolygonContour( PolygonToCreate.PerimeterVertices, NewEdgeIDsForContour, PerimeterVertexInstances );
 
+			TArray<FEdgeID> CreatedEdgeIDs;
 			FPolygonID PolygonID = PolygonToCreate.OriginalPolygonID;
 			if( PolygonID != FPolygonID::Invalid )
 			{
-				GetMeshDescription()->CreatePolygonWithID( PolygonID, PolygonToCreate.PolygonGroupID, PerimeterVertexInstances, &NewEdgeIDsForContour );
+				GetMeshDescription()->CreatePolygonWithID( PolygonID, PolygonToCreate.PolygonGroupID, PerimeterVertexInstances, &CreatedEdgeIDs );
 			}
 			else
 			{
-				PolygonID = GetMeshDescription()->CreatePolygon( PolygonToCreate.PolygonGroupID, PerimeterVertexInstances, &NewEdgeIDsForContour );
+				PolygonID = GetMeshDescription()->CreatePolygon( PolygonToCreate.PolygonGroupID, PerimeterVertexInstances, &CreatedEdgeIDs );
 			}
+			check( CreatedEdgeIDs.Num() == 0 );
 
 			OutNewPolygonIDs.Add( PolygonID );
 			OutNewEdgeIDs.Append( NewEdgeIDsForContour );
@@ -4150,19 +4082,19 @@ void UEditableMesh::CreatePolygons( const TArray<FPolygonToCreate>& PolygonsToCr
 	// Generate triangles for the new polygon
 	PolygonsPendingTriangulation.Append( OutNewPolygonIDs );
 
-	// If any new edges were automatically created when creating the polygons, add an action to the undo stack to delete them here
-	if( OutNewEdgeIDs.Num() > 0 )
-	{
-		FDeleteEdgesChangeInput RevertInput;
-		RevertInput.bDeleteOrphanedVertices = false;
-		RevertInput.EdgeIDsToDelete.Reserve( OutNewEdgeIDs.Num() );
-		for( int32 EdgeNumber = OutNewEdgeIDs.Num() - 1; EdgeNumber >= 0; --EdgeNumber )
-		{
-			RevertInput.EdgeIDsToDelete.Add( OutNewEdgeIDs[ EdgeNumber ] );
-		}
+	//// If any new edges were automatically created when creating the polygons, add an action to the undo stack to delete them here
+	//if( OutNewEdgeIDs.Num() > 0 )
+	//{
+	//	FDeleteEdgesChangeInput RevertInput;
+	//	RevertInput.bDeleteOrphanedVertices = false;
+	//	RevertInput.EdgeIDsToDelete.Reserve( OutNewEdgeIDs.Num() );
+	//	for( int32 EdgeNumber = OutNewEdgeIDs.Num() - 1; EdgeNumber >= 0; --EdgeNumber )
+	//	{
+	//		RevertInput.EdgeIDsToDelete.Add( OutNewEdgeIDs[ EdgeNumber ] );
+	//	}
 
-		AddUndo( MakeUnique<FDeleteEdgesChange>( MoveTemp( RevertInput ) ) );
-	}
+	//	AddUndo( MakeUnique<FDeleteEdgesChange>( MoveTemp( RevertInput ) ) );
+	//}
 
 	// NOTE: We iterate backwards, to delete polygons in the opposite order that we added them
 	{
@@ -4184,32 +4116,10 @@ void UEditableMesh::CreatePolygons( const TArray<FPolygonToCreate>& PolygonsToCr
 }
 
 
-void UEditableMesh::BackupPolygonContour( const FMeshPolygonContour& Contour, TArray<FVertexAndAttributes>& OutVerticesAndAttributes )
-{
-	const FVertexInstanceArray& VertexInstances = GetMeshDescription()->VertexInstances();
-
-	OutVerticesAndAttributes.Reserve( Contour.VertexInstanceIDs.Num() );
-	for( const FVertexInstanceID VertexInstanceID : Contour.VertexInstanceIDs )
-	{
-		const FMeshVertexInstance& VertexInstance = VertexInstances[ VertexInstanceID ];
-
-		OutVerticesAndAttributes.Emplace();
-		FVertexAndAttributes& VertexAndAttributes = OutVerticesAndAttributes.Last();
-
-		// We rely on undoing recreating vertex instances, therefore we only need pass their IDs.
-		VertexAndAttributes.VertexInstanceID = VertexInstanceID;
-		VertexAndAttributes.VertexID = FVertexID::Invalid;
-	}
-}
-
-
 void UEditableMesh::DeletePolygons( const TArray<FPolygonID>& PolygonIDsToDelete, const bool bDeleteOrphanedEdges, const bool bDeleteOrphanedVertices, const bool bDeleteOrphanedVertexInstances, const bool bDeleteEmptyPolygonGroups )
 {
 	EM_ENTER( TEXT( "DeletePolygons: %s" ), *LogHelpers::ArrayToString( PolygonIDsToDelete ) );
 
-	FPolygonArray& Polygons = GetMeshDescription()->Polygons();
-	FPolygonGroupArray& PolygonGroups = GetMeshDescription()->PolygonGroups();
-
 	// Back everything up
 	{
 		FCreatePolygonsChangeInput RevertInput;
@@ -4219,15 +4129,24 @@ void UEditableMesh::DeletePolygons( const TArray<FPolygonID>& PolygonIDsToDelete
 		for( int32 PolygonNumber = PolygonIDsToDelete.Num() - 1; PolygonNumber >= 0; --PolygonNumber )
 		{
 			const FPolygonID PolygonID = PolygonIDsToDelete[ PolygonNumber ];
-			const FMeshPolygon& Polygon = Polygons[ PolygonID ];
 
 			RevertInput.PolygonsToCreate.Emplace();
 			FPolygonToCreate& PolygonToCreate = RevertInput.PolygonsToCreate.Last();
 
-			PolygonToCreate.PolygonGroupID = Polygon.PolygonGroupID;
+			PolygonToCreate.PolygonGroupID = GetMeshDescription()->GetPolygonPolygonGroup( PolygonID );
 			PolygonToCreate.OriginalPolygonID = PolygonID;
 
-			BackupPolygonContour( Polygon.PerimeterContour, PolygonToCreate.PerimeterVertices );
+			PolygonToCreate.PerimeterVertices.Reserve( GetMeshDescription()->GetPolygonVertexInstances( PolygonID ).Num() );
+			for( const FVertexInstanceID VertexInstanceID : GetMeshDescription()->GetPolygonVertexInstances( PolygonID ) )
+			{
+				PolygonToCreate.PerimeterVertices.Emplace();
+				FVertexAndAttributes& VertexAndAttributes = PolygonToCreate.PerimeterVertices.Last();
+
+				// We rely on undoing recreating vertex instances, therefore we only need pass their IDs.
+				VertexAndAttributes.VertexInstanceID = VertexInstanceID;
+				VertexAndAttributes.VertexID = FVertexID::Invalid;
+			}
+
 		}
 
 		AddUndo( MakeUnique<FCreatePolygonsChange>( MoveTemp( RevertInput ) ) );
@@ -4527,7 +4446,6 @@ void UEditableMesh::ChangePolygonsVertexInstances( const TArray<FChangeVertexIns
 		{
 			const FChangeVertexInstancesForPolygon& VertexInstancesForPolygon = VertexInstancesForPolygons[ Index ];
 			const FPolygonID PolygonID = VertexInstancesForPolygon.PolygonID;
-			const FMeshPolygon& Polygon = Polygons[ PolygonID ];
 
 			RevertInput.VertexInstancesForPolygons.Emplace();
 			FChangeVertexInstancesForPolygon& RevertVertexInstancesForPolygon = RevertInput.VertexInstancesForPolygons.Last();
@@ -4538,7 +4456,7 @@ void UEditableMesh::ChangePolygonsVertexInstances( const TArray<FChangeVertexIns
 				RevertVertexInstancesForPolygon.PerimeterVertexIndicesAndInstanceIDs.Emplace();
 				FVertexIndexAndInstanceID& RevertIndexAndInstance = RevertVertexInstancesForPolygon.PerimeterVertexIndicesAndInstanceIDs.Last();
 				RevertIndexAndInstance.ContourIndex = IndexAndInstance.ContourIndex;
-				RevertIndexAndInstance.VertexInstanceID = Polygon.PerimeterContour.VertexInstanceIDs[ IndexAndInstance.ContourIndex ];
+				RevertIndexAndInstance.VertexInstanceID = GetPolygonPerimeterVertexInstance( PolygonID, IndexAndInstance.ContourIndex );
 			}
 		}
 
@@ -4554,32 +4472,10 @@ void UEditableMesh::ChangePolygonsVertexInstances( const TArray<FChangeVertexIns
 		{
 			const FPolygonID PolygonID = VertexInstancesForPolygon.PolygonID;
 			PolygonIDs.Add( PolygonID );
-			FMeshPolygon& Polygon = Polygons[ PolygonID ];
 
 			for( const FVertexIndexAndInstanceID& IndexAndInstance : VertexInstancesForPolygon.PerimeterVertexIndicesAndInstanceIDs )
 			{
-				// Disconnect old vertex instance from polygon, and connect new one
-				const FVertexInstanceID OldVertexInstanceID = Polygon.PerimeterContour.VertexInstanceIDs[ IndexAndInstance.ContourIndex ];
-				FMeshVertexInstance& OldVertexInstance = VertexInstances[ OldVertexInstanceID ];
-				verify( OldVertexInstance.ConnectedPolygons.Remove( PolygonID ) == 1 );
-
-				FMeshVertexInstance& NewVertexInstance = VertexInstances[ IndexAndInstance.VertexInstanceID ];
-				check( !NewVertexInstance.ConnectedPolygons.Contains( PolygonID ) );
-				NewVertexInstance.ConnectedPolygons.Add( PolygonID );
-
-				Polygon.PerimeterContour.VertexInstanceIDs[ IndexAndInstance.ContourIndex ] = IndexAndInstance.VertexInstanceID;
-
-				// Fix up triangle list
-				for( FMeshTriangle& Triangle : Polygon.Triangles )
-				{
-					for( int32 VertexIndex = 0; VertexIndex < 3; ++VertexIndex )
-					{
-						if( Triangle.GetVertexInstanceID( VertexIndex ) == OldVertexInstanceID )
-						{
-							Triangle.SetVertexInstanceID( VertexIndex, IndexAndInstance.VertexInstanceID );
-						}
-					}
-				}
+				GetMeshDescription()->SetPolygonVertexInstance( PolygonID, IndexAndInstance.ContourIndex, IndexAndInstance.VertexInstanceID );
 			}
 		}
 	}
@@ -4596,12 +4492,9 @@ void UEditableMesh::ChangePolygonsVertexInstances( const TArray<FChangeVertexIns
 
 FVertexInstanceID UEditableMesh::GetVertexInstanceInPolygonForVertex( const FPolygonID PolygonID, const FVertexID VertexID ) const
 {
-	const FVertexInstanceArray& VertexInstances = GetMeshDescription()->VertexInstances();
-
-	for( const FVertexInstanceID VertexInstanceID : GetMeshDescription()->GetVertex( VertexID ).VertexInstanceIDs )
+	for( const FVertexInstanceID VertexInstanceID : GetMeshDescription()->GetVertexVertexInstances( VertexID ) )
 	{
-		const FMeshVertexInstance& VertexInstance = VertexInstances[ VertexInstanceID ];
-		if( VertexInstance.ConnectedPolygons.Contains( PolygonID ) )
+		if( GetMeshDescription()->GetVertexInstanceConnectedPolygons( VertexInstanceID ).Contains( PolygonID ) )
 		{
 			return VertexInstanceID;
 		}
@@ -4616,7 +4509,7 @@ void UEditableMesh::GetConnectedSoftEdges( const FVertexID VertexID, TArray<FEdg
 	OutConnectedSoftEdges.Reset();
 
 	TEdgeAttributesConstRef<bool> EdgeHardnesses = GetMeshDescription()->EdgeAttributes().GetAttributesRef<bool>( MeshAttribute::Edge::IsHard );
-	for( const FEdgeID ConnectedEdgeID : GetMeshDescription()->GetVertex( VertexID ).ConnectedEdgeIDs )
+	for( const FEdgeID ConnectedEdgeID : GetMeshDescription()->GetVertexConnectedEdges( VertexID ) )
 	{
 		if( !EdgeHardnesses[ ConnectedEdgeID ] )
 		{
@@ -4630,23 +4523,22 @@ void UEditableMesh::SetPolygonsVertexAttributes( const TArray<FVertexAttributesF
 {
 	EM_ENTER( TEXT( "SetPolygonsVertexAttributes: %s" ), *LogHelpers::ArrayToString( VertexAttributesForPolygons ) );
 
-	FPolygonArray& Polygons = GetMeshDescription()->Polygons();
-
 	for( const FVertexAttributesForPolygon& VertexAttributesForPolygon : VertexAttributesForPolygons )
 	{
 		const FPolygonID PolygonID = VertexAttributesForPolygon.PolygonID;
-		FMeshPolygon& Polygon = Polygons[ PolygonID ];
 
-		SetPolygonContourVertexAttributes( Polygon.PerimeterContour, PolygonID, VertexAttributesForPolygon.PerimeterVertexAttributeLists );
+		SetPolygonContourVertexAttributes( PolygonID, VertexAttributesForPolygon.PerimeterVertexAttributeLists );
 	}
 
 	EM_EXIT( TEXT( "SetPolygonsVertexAttributes returned" ) );
 }
 
 
-void UEditableMesh::SetPolygonContourVertexAttributes( FMeshPolygonContour& Contour, const FPolygonID PolygonID, const TArray<FMeshElementAttributeList>& AttributeLists )
+void UEditableMesh::SetPolygonContourVertexAttributes( const FPolygonID PolygonID, const TArray<FMeshElementAttributeList>& AttributeLists )
 {
-	const int32 NumContourVertices = Contour.VertexInstanceIDs.Num();
+	const TArray<FVertexInstanceID>& ContourVertexInstanceIDs = GetMeshDescription()->GetPolygonVertexInstances( PolygonID );
+
+	const int32 NumContourVertices = ContourVertexInstanceIDs.Num();
 	check( AttributeLists.Num() == NumContourVertices );
 
 	// Iterate round all polygons in the contour
@@ -4661,7 +4553,7 @@ void UEditableMesh::SetPolygonContourVertexAttributes( FMeshPolygonContour& Cont
 		}
 
 		// Get vertex instance and vertex.
-		const FVertexInstanceID VertexInstanceID = Contour.VertexInstanceIDs[ Index ];
+		const FVertexInstanceID VertexInstanceID = ContourVertexInstanceIDs[ Index ];
 		const FVertexID VertexID = GetVertexInstanceVertex( VertexInstanceID );
 
 		const TArray<FPolygonID>& ConnectedPolygons = GetMeshDescription()->GetVertexInstanceConnectedPolygons( VertexInstanceID );
@@ -4759,8 +4651,8 @@ void UEditableMesh::TryToRemovePolygonEdge( const FEdgeID EdgeID, bool& bOutWasE
 			const FPolygonID PolygonAID = GetEdgeConnectedPolygon( EdgeID, 0 );
 			const FPolygonID PolygonBID = GetEdgeConnectedPolygon( EdgeID, 1 );
 
-			const TArray<FVertexInstanceID>& PolygonAVertexInstanceIDs = GetMeshDescription()->GetPolygonPerimeterVertexInstances( PolygonAID );
-			const TArray<FVertexInstanceID>& PolygonBVertexInstanceIDs = GetMeshDescription()->GetPolygonPerimeterVertexInstances( PolygonBID );
+			const TArray<FVertexInstanceID>& PolygonAVertexInstanceIDs = GetMeshDescription()->GetPolygonVertexInstances( PolygonAID );
+			const TArray<FVertexInstanceID>& PolygonBVertexInstanceIDs = GetMeshDescription()->GetPolygonVertexInstances( PolygonBID );
 
 			// If the polygons are in different polygon groups, we can't remove the edge because we can't determine
 			// which polygon group the replacing polygon should belong to
@@ -4929,33 +4821,53 @@ void UEditableMesh::TryToRemoveVertex( const FVertexID VertexID, bool& bOutWasVe
 		static TArray<FPolygonID> NewEdgeConnectedPolygons;
 		GetVertexConnectedPolygons( VertexID, /* Out */ NewEdgeConnectedPolygons );
 
+		TArray<FPolygonToCreate> PolygonsToCreate;
+		PolygonsToCreate.Reserve( NewEdgeConnectedPolygons.Num() );
+
 		// Remove the vertex from its connected polygons
 		{
 			for( const FPolygonID PolygonID : NewEdgeConnectedPolygons )
 			{
+				FPolygonToCreate& PolygonToCreate = PolygonsToCreate.Emplace_GetRef();
+
+				const FPolygonGroupID PolygonGroupID = GetMeshDescription()->GetPolygonPolygonGroup( PolygonID );
+				TArray<FVertexInstanceID> PerimeterVertexInstanceIDs = GetMeshDescription()->GetPolygonVertexInstances( PolygonID );
+
+				PolygonToCreate.PolygonGroupID = PolygonGroupID;
+				PolygonToCreate.PolygonEdgeHardness = EPolygonEdgeHardness::NewEdgesSoft;
+				PolygonToCreate.PerimeterVertices.Reserve( PerimeterVertexInstanceIDs.Num() - 1 );
+
 				const int32 PolygonVertexNumber = FindPolygonPerimeterVertexNumberForVertex( PolygonID, VertexID );
 				check( PolygonVertexNumber != INDEX_NONE );
-				const bool bDeleteOrphanedVertexInstances = false;
-				RemovePolygonPerimeterVertices( PolygonID, PolygonVertexNumber, 1, bDeleteOrphanedVertexInstances );
-			}
-		}
 
-		// Delete the two edges
-		{
-			static TArray<FEdgeID> EdgeIDsToDelete;
-			EdgeIDsToDelete.Reset();
-			for( int32 EdgeNumber = 0; EdgeNumber < ConnectedEdgeCount; ++EdgeNumber )
-			{
-				EdgeIDsToDelete.Add( GetVertexConnectedEdge( VertexID, EdgeNumber ) );
-			}
+				for( int32 VertexNumber = 0; VertexNumber < PerimeterVertexInstanceIDs.Num(); ++VertexNumber )
+				{
+					if( VertexNumber != PolygonVertexNumber )
+					{
+						FVertexAndAttributes& VertexAndAttributes = PolygonToCreate.PerimeterVertices.Emplace_GetRef();
+						VertexAndAttributes.VertexInstanceID = PerimeterVertexInstanceIDs[ VertexNumber ];
+					}
+				}
 
-			// NOTE: We can't delete the orphan vertex yet because the polygon triangles are still referencing
-			// its rendering vertices.  We'll delete the edges, retriangulate, then delete the vertex afterwards.
-			const bool bDeleteOrphanedVertices = false;
-			DeleteEdges( EdgeIDsToDelete, bDeleteOrphanedVertices );
+//				const bool bDeleteOrphanedVertexInstances = false;
+//				RemovePolygonPerimeterVertices( PolygonID, PolygonVertexNumber, 1, bDeleteOrphanedVertexInstances );
+			}
 		}
 
-		FEdgeID NewEdgeID = FEdgeID::Invalid;
+		//// Delete the two edges
+		//{
+		//	static TArray<FEdgeID> EdgeIDsToDelete;
+		//	EdgeIDsToDelete.Reset();
+		//	for( int32 EdgeNumber = 0; EdgeNumber < ConnectedEdgeCount; ++EdgeNumber )
+		//	{
+		//		EdgeIDsToDelete.Add( GetVertexConnectedEdge( VertexID, EdgeNumber ) );
+		//	}
+
+		//	// NOTE: We can't delete the orphan vertex yet because the polygon triangles are still referencing
+		//	// its rendering vertices.  We'll delete the edges, retriangulate, then delete the vertex afterwards.
+		//	const bool bDeleteOrphanedVertices = false;
+		//	DeleteEdges( EdgeIDsToDelete, bDeleteOrphanedVertices );
+		//}
 
 		// Create a new edge to replace the vertex and two edges we deleted
 		{
@@ -4965,14 +4877,28 @@ void UEditableMesh::TryToRemoveVertex( const FVertexID VertexID, bool& bOutWasVe
 			FEdgeToCreate& EdgeToCreate = *new( EdgesToCreate ) FEdgeToCreate;
 			EdgeToCreate.VertexID0 = NewEdgeVertexIDs[ 0 ];
 			EdgeToCreate.VertexID1 = NewEdgeVertexIDs[ 1 ];
-			EdgeToCreate.ConnectedPolygons = NewEdgeConnectedPolygons;
 			EdgeToCreate.EdgeAttributes = EdgeAttributeList;
 
-			static TArray<FEdgeID> NewEdgeIDs;
-			NewEdgeIDs.Reset();
+			TArray<FEdgeID> NewEdgeIDs;
 			CreateEdges( EdgesToCreate, /* Out */ NewEdgeIDs );
 
-			NewEdgeID = NewEdgeIDs[ 0 ];
+			OutNewEdgeID = NewEdgeIDs[ 0 ];
+		}
+
+		// Delete the old polygons - this will also delete the old edges, vertex instances and vertices
+		{
+			const bool bDeleteOrphanEdges = true;
+			const bool bDeleteOrphanVertices = true;
+			const bool bDeleteOrphanVertexInstances = true;
+			const bool bDeleteEmptyPolygonGroups = false;
+			DeletePolygons( NewEdgeConnectedPolygons, bDeleteOrphanEdges, bDeleteOrphanVertices, bDeleteOrphanVertexInstances, bDeleteEmptyPolygonGroups );
+		}
+
+		// Create new polygons with the vertex removed
+		{
+			TArray<FPolygonID> NewPolygonIDs;
+			TArray<FEdgeID> NewEdgeIDs;
+			CreatePolygons( PolygonsToCreate, NewPolygonIDs, NewEdgeIDs );
 		}
 
 		// Update the normals of the affected polygons
@@ -4981,20 +4907,16 @@ void UEditableMesh::TryToRemoveVertex( const FVertexID VertexID, bool& bOutWasVe
 		// Retriangulate all of the affected polygons
 		PolygonsPendingTriangulation.Append( NewEdgeConnectedPolygons );
 
-		// Delete the vertex instances and subsequently orphaned vertex
-		{
-			FMeshVertex& Vertex = GetMeshDescription()->GetVertex( VertexID );
-			const bool bDeleteOrphanedVertices = true;
-
-			// Take a copy of the array, because it will be modified by the DeleteVertexInstances call
-			static TArray<FVertexInstanceID> VertexInstanceIDs;
-			VertexInstanceIDs = Vertex.VertexInstanceIDs;
+		//// Delete the vertex instances and subsequently orphaned vertex
+		//{
+		//	// Take a copy of the array, because it will be modified by the DeleteVertexInstances call
+		//	TArray<FVertexInstanceID> VertexInstanceIDs = GetMeshDescription()->GetVertexVertexInstances( VertexID );
 
-			DeleteVertexInstances( VertexInstanceIDs, bDeleteOrphanedVertices );
-		}
+		//	const bool bDeleteOrphanedVertices = true;
+		//	DeleteVertexInstances( VertexInstanceIDs, bDeleteOrphanedVertices );
+		//}
 
 		bOutWasVertexRemoved = true;
-		OutNewEdgeID = NewEdgeID;
 	}
 
 	EM_EXIT( TEXT( "TryToRemoveVertex returned %s %s" ), *LogHelpers::BoolToString( bOutWasVertexRemoved ), *OutNewEdgeID.ToString() );
@@ -5282,7 +5204,7 @@ void UEditableMesh::ExtrudePolygons( const TArray<FPolygonID>& PolygonIDs, const
 			NewFrontPolygonVertices.Reset( PolygonVertexIDs.Num() );
 			NewFrontPolygonVertices.SetNum( PolygonVertexIDs.Num(), false );
 
-			const TArray<FVertexInstanceID>& VertexInstanceIDs = GetMeshDescription()->GetPolygonPerimeterVertexInstances( PolygonID );
+			const TArray<FVertexInstanceID>& VertexInstanceIDs = GetMeshDescription()->GetPolygonVertexInstances( PolygonID );
 
 			// Map all of the polygon's vertex IDs to their extruded counterparts to create the new polygon perimeter
 			for( int32 PolygonVertexNumber = 0; PolygonVertexNumber < PolygonVertexIDs.Num(); ++PolygonVertexNumber )
@@ -5442,7 +5364,7 @@ void UEditableMesh::ExtendEdges( const TArray<FEdgeID>& EdgeIDs, const bool bWel
 			BackupAllAttributes( EdgeToCreate.EdgeAttributes, GetMeshDescription()->EdgeAttributes(), EdgeID );
 
 			// We're not connected to any polygons yet.  That will come later.
-			EdgeToCreate.ConnectedPolygons.Reset();
+			//EdgeToCreate.ConnectedPolygons.Reset();
 		}
 
 		CreateEdges( EdgesToCreate, /* Out */ OutNewExtendedEdgeIDs );
@@ -5748,12 +5670,12 @@ void UEditableMesh::BevelOrInsetPolygons( const TArray<FPolygonID>& PolygonIDs,
 		static TArray<FVertexID> PerimeterVertexIDs;
 		GetPolygonPerimeterVertices( PolygonID, /* Out */ PerimeterVertexIDs );
 
-		const TArray<FVertexInstanceID>& PerimeterVertexInstanceIDs = GetMeshDescription()->GetPolygonPerimeterVertexInstances( PolygonID );
+		const TArray<FVertexInstanceID>& PerimeterVertexInstanceIDs = GetMeshDescription()->GetPolygonVertexInstances( PolygonID );
 
 		static TArray<FVertexToCreate> VerticesToCreate;
 		VerticesToCreate.Reset( PerimeterVertexIDs.Num() );
 
-		static TArray<TTuple<FMeshTriangle, FVector>> BarycentricWeightsForNewVertices;
+		static TArray<TTuple<FTriangleID, FVector>> BarycentricWeightsForNewVertices;
 		BarycentricWeightsForNewVertices.Reset();
 		BarycentricWeightsForNewVertices.Reserve( PerimeterVertexIDs.Num() );
 
@@ -5814,8 +5736,8 @@ void UEditableMesh::BevelOrInsetPolygons( const TArray<FPolygonID>& PolygonIDs,
 			FMeshTriangle Triangle;
 			FVector TriangleVertexWeights;
 
-			bool bPointInPolygon = ComputeBarycentricWeightForPointOnPolygon( PolygonID, InsetVertexPosition, Triangle, TriangleVertexWeights );
-			BarycentricWeightsForNewVertices.Emplace( Triangle, TriangleVertexWeights );
+			const FTriangleID TriangleID = ComputeBarycentricWeightForPointOnPolygon( PolygonID, InsetVertexPosition, TriangleVertexWeights );
+			BarycentricWeightsForNewVertices.Emplace( TriangleID, TriangleVertexWeights );
 
 			// If we're beveling, go ahead and move the original polygon perimeter vertices
 			if( bShouldBevel )
@@ -5879,13 +5801,13 @@ void UEditableMesh::BevelOrInsetPolygons( const TArray<FPolygonID>& PolygonIDs,
 				BackupAllAttributes( PolygonToCreate.PerimeterVertices[ 1 ].PolygonVertexAttributes, GetMeshDescription()->VertexInstanceAttributes(), RightOriginalVertexInstanceID );
 
 				// Inset right
-				if( BarycentricWeightsForNewVertices[ RightVertexNumber ].Get<0>().VertexInstanceID0 != FVertexInstanceID::Invalid )
+				if( BarycentricWeightsForNewVertices[ RightVertexNumber ].Get<0>() != FTriangleID::Invalid )
 				{
 					InterpAllAttributes( PolygonToCreate.PerimeterVertices[ 2 ].PolygonVertexAttributes,
 										 GetMeshDescription()->VertexInstanceAttributes(),
-										 BarycentricWeightsForNewVertices[ RightVertexNumber ].Get<0>().VertexInstanceID0,
-										 BarycentricWeightsForNewVertices[ RightVertexNumber ].Get<0>().VertexInstanceID1,
-										 BarycentricWeightsForNewVertices[ RightVertexNumber ].Get<0>().VertexInstanceID2,
+										 GetMeshDescription()->GetTriangleVertexInstance( BarycentricWeightsForNewVertices[ RightVertexNumber ].Get<0>(), 0 ),
+										 GetMeshDescription()->GetTriangleVertexInstance( BarycentricWeightsForNewVertices[ RightVertexNumber ].Get<0>(), 1 ),
+										 GetMeshDescription()->GetTriangleVertexInstance( BarycentricWeightsForNewVertices[ RightVertexNumber ].Get<0>(), 2 ),
 										 BarycentricWeightsForNewVertices[ RightVertexNumber ].Get<1>() );
 				}
 				else
@@ -5894,13 +5816,13 @@ void UEditableMesh::BevelOrInsetPolygons( const TArray<FPolygonID>& PolygonIDs,
 				}
 
 				// Inset left
-				if( BarycentricWeightsForNewVertices[ LeftVertexNumber ].Get<0>().VertexInstanceID0 != FVertexInstanceID::Invalid )
+				if( BarycentricWeightsForNewVertices[ LeftVertexNumber ].Get<0>() != FTriangleID::Invalid )
 				{
 					InterpAllAttributes( PolygonToCreate.PerimeterVertices[ 3 ].PolygonVertexAttributes,
 										 GetMeshDescription()->VertexInstanceAttributes(),
-										 BarycentricWeightsForNewVertices[ LeftVertexNumber ].Get<0>().VertexInstanceID0,
-										 BarycentricWeightsForNewVertices[ LeftVertexNumber ].Get<0>().VertexInstanceID1,
-										 BarycentricWeightsForNewVertices[ LeftVertexNumber ].Get<0>().VertexInstanceID2,
+										 GetMeshDescription()->GetTriangleVertexInstance( BarycentricWeightsForNewVertices[ LeftVertexNumber ].Get<0>(), 0 ),
+										 GetMeshDescription()->GetTriangleVertexInstance( BarycentricWeightsForNewVertices[ LeftVertexNumber ].Get<0>(), 1 ),
+										 GetMeshDescription()->GetTriangleVertexInstance( BarycentricWeightsForNewVertices[ LeftVertexNumber ].Get<0>(), 2 ),
 										 BarycentricWeightsForNewVertices[ LeftVertexNumber ].Get<1>() );
 				}
 				else
@@ -5929,13 +5851,13 @@ void UEditableMesh::BevelOrInsetPolygons( const TArray<FPolygonID>& PolygonIDs,
 				FVertexAndAttributes& NewPerimeterVertex = PolygonToCreate.PerimeterVertices[ NewVertexNumber ];
 				NewPerimeterVertex.VertexID = NewVertexID;
 
-				if( BarycentricWeightsForNewVertices[ NewVertexNumber ].Get<0>().VertexInstanceID0 != FVertexInstanceID::Invalid )
+				if( BarycentricWeightsForNewVertices[ NewVertexNumber ].Get<0>() != FTriangleID::Invalid )
 				{
 					InterpAllAttributes( NewPerimeterVertex.PolygonVertexAttributes,
 										 GetMeshDescription()->VertexInstanceAttributes(),
-										 BarycentricWeightsForNewVertices[ NewVertexNumber ].Get<0>().VertexInstanceID0,
-										 BarycentricWeightsForNewVertices[ NewVertexNumber ].Get<0>().VertexInstanceID1,
-										 BarycentricWeightsForNewVertices[ NewVertexNumber ].Get<0>().VertexInstanceID2,
+										 GetMeshDescription()->GetTriangleVertexInstance( BarycentricWeightsForNewVertices[ NewVertexNumber ].Get<0>(), 0 ),
+										 GetMeshDescription()->GetTriangleVertexInstance( BarycentricWeightsForNewVertices[ NewVertexNumber ].Get<0>(), 1 ),
+										 GetMeshDescription()->GetTriangleVertexInstance( BarycentricWeightsForNewVertices[ NewVertexNumber ].Get<0>(), 2 ),
 										 BarycentricWeightsForNewVertices[ NewVertexNumber ].Get<1>() );
 				}
 				else
@@ -6000,73 +5922,28 @@ void UEditableMesh::BevelPolygons( const TArray<FPolygonID>& PolygonIDs, const f
 }
 
 
-float UEditableMesh::GetPolygonCornerAngleForVertex( const FPolygonID PolygonID, const FVertexID VertexID ) const
-{
-	const FMeshPolygon& Polygon = GetMeshDescription()->GetPolygon( PolygonID );
-
-	// Lambda function which returns the inner angle at a given index on a polygon contour
-	auto GetContourAngle = [ this ]( const FMeshPolygonContour& Contour, const int32 ContourIndex )
-	{
-		const int32 NumVertices = Contour.VertexInstanceIDs.Num();
-
-		const int32 PrevIndex = ( ContourIndex + NumVertices - 1 ) % NumVertices;
-		const int32 NextIndex = ( ContourIndex + 1 ) % NumVertices;
-
-		const FVertexID PrevVertexID = GetMeshDescription()->GetVertexInstanceVertex( Contour.VertexInstanceIDs[ PrevIndex ] );
-		const FVertexID ThisVertexID = GetMeshDescription()->GetVertexInstanceVertex( Contour.VertexInstanceIDs[ ContourIndex ] );
-		const FVertexID NextVertexID = GetMeshDescription()->GetVertexInstanceVertex( Contour.VertexInstanceIDs[ NextIndex ] );
-
-		TVertexAttributesConstRef<FVector> VertexPositions = GetMeshDescription()->VertexAttributes().GetAttributesRef<FVector>( MeshAttribute::Vertex::Position );
-
-		const FVector PrevVertexPosition = VertexPositions[ PrevVertexID ];
-		const FVector ThisVertexPosition = VertexPositions[ ThisVertexID ];
-		const FVector NextVertexPosition = VertexPositions[ NextVertexID ];
-
-		const FVector Direction1 = ( PrevVertexPosition - ThisVertexPosition ).GetSafeNormal();
-		const FVector Direction2 = ( NextVertexPosition - ThisVertexPosition ).GetSafeNormal();
-
-		return FMath::Acos( FVector::DotProduct( Direction1, Direction2 ) );
-	};
-
-	const FVertexInstanceArray& VertexInstances = GetMeshDescription()->VertexInstances();
-	auto IsVertexInstancedFromThisVertex = [ &VertexInstances, VertexID ]( const FVertexInstanceID VertexInstanceID )
-	{
-		return VertexInstances[ VertexInstanceID ].VertexID == VertexID;
-	};
-
-	// First look for the vertex instance in the perimeter
-	int32 ContourIndex = Polygon.PerimeterContour.VertexInstanceIDs.IndexOfByPredicate( IsVertexInstancedFromThisVertex );
-	if( ContourIndex != INDEX_NONE )
-	{
-		// Return the internal angle if found
-		return GetContourAngle( Polygon.PerimeterContour, ContourIndex );
-	}
-
-	// Found nothing; return 0
-	return 0.0f;
-}
-
-
 void UEditableMesh::GeneratePolygonTangentsAndNormals( const TArray<FPolygonID>& PolygonIDs )
 {
-	const TVertexAttributesRef<FVector> VertexPositions = GetMeshDescription()->VertexAttributes().GetAttributesRef<FVector>( MeshAttribute::Vertex::Position );
+	FMeshDescription* MD = GetMeshDescription();
+
+	const TVertexAttributesRef<FVector> VertexPositions = MD->VertexAttributes().GetAttributesRef<FVector>( MeshAttribute::Vertex::Position );
 
 	// @todo mesheditor: currently hardcoded to calculate the tangent basis from UV0.
-	const TVertexInstanceAttributesRef<FVector2D> VertexUVs = GetMeshDescription()->VertexInstanceAttributes().GetAttributesRef<FVector2D>( MeshAttribute::VertexInstance::TextureCoordinate );
+	const TVertexInstanceAttributesRef<FVector2D> VertexUVs = MD->VertexInstanceAttributes().GetAttributesRef<FVector2D>( MeshAttribute::VertexInstance::TextureCoordinate );
 
-	TPolygonAttributesRef<FVector> PolygonNormals = GetMeshDescription()->PolygonAttributes().GetAttributesRef<FVector>( MeshAttribute::Polygon::Normal );
-	TPolygonAttributesRef<FVector> PolygonTangents = GetMeshDescription()->PolygonAttributes().GetAttributesRef<FVector>( MeshAttribute::Polygon::Tangent );
-	TPolygonAttributesRef<FVector> PolygonBinormals = GetMeshDescription()->PolygonAttributes().GetAttributesRef<FVector>( MeshAttribute::Polygon::Binormal );
-	TPolygonAttributesRef<FVector> PolygonCenters = GetMeshDescription()->PolygonAttributes().GetAttributesRef<FVector>( MeshAttribute::Polygon::Center );
+	TPolygonAttributesRef<FVector> PolygonNormals = MD->PolygonAttributes().GetAttributesRef<FVector>( MeshAttribute::Polygon::Normal );
+	TPolygonAttributesRef<FVector> PolygonTangents = MD->PolygonAttributes().GetAttributesRef<FVector>( MeshAttribute::Polygon::Tangent );
+	TPolygonAttributesRef<FVector> PolygonBinormals = MD->PolygonAttributes().GetAttributesRef<FVector>( MeshAttribute::Polygon::Binormal );
+	TPolygonAttributesRef<FVector> PolygonCenters = MD->PolygonAttributes().GetAttributesRef<FVector>( MeshAttribute::Polygon::Center );
 
 	for( const FPolygonID PolygonID : PolygonIDs )
 	{
 		// Calculate the center of this polygon
 		FVector Center = FVector::ZeroVector;
-		const TArray<FVertexInstanceID>& VertexInstanceIDs = GetMeshDescription()->GetPolygonPerimeterVertexInstances( PolygonID );
+		const TArray<FVertexInstanceID>& VertexInstanceIDs = MD->GetPolygonVertexInstances( PolygonID );
 		for( const FVertexInstanceID VertexInstanceID : VertexInstanceIDs )
 		{
-			Center += VertexPositions[ GetMeshDescription()->GetVertexInstanceVertex( VertexInstanceID ) ];
+			Center += VertexPositions[ MD->GetVertexInstanceVertex( VertexInstanceID ) ];
 		}
 		Center /= float( VertexInstanceIDs.Num() );
 
@@ -6075,17 +5952,21 @@ void UEditableMesh::GeneratePolygonTangentsAndNormals( const TArray<FPolygonID>&
 		FVector Tangent = FVector::ZeroVector;
 		FVector Binormal = FVector::ZeroVector;
 
-		for( const FMeshTriangle& Triangle : GetMeshDescription()->GetPolygonTriangles( PolygonID ) )
+		for( const FTriangleID TriangleID : MD->GetPolygonTriangleIDs( PolygonID ) )
 		{
-			const FVertexID VertexID0 = GetMeshDescription()->GetVertexInstanceVertex( Triangle.VertexInstanceID0 );
-			const FVertexID VertexID1 = GetMeshDescription()->GetVertexInstanceVertex( Triangle.VertexInstanceID1 );
-			const FVertexID VertexID2 = GetMeshDescription()->GetVertexInstanceVertex( Triangle.VertexInstanceID2 );
+			const FVertexInstanceID VertexInstanceID0 = MD->GetTriangleVertexInstance( TriangleID, 0 );
+			const FVertexInstanceID VertexInstanceID1 = MD->GetTriangleVertexInstance( TriangleID, 1 );
+			const FVertexInstanceID VertexInstanceID2 = MD->GetTriangleVertexInstance( TriangleID, 2 );
+
+			const FVertexID VertexID0 = MD->GetVertexInstanceVertex( VertexInstanceID0 );
+			const FVertexID VertexID1 = MD->GetVertexInstanceVertex( VertexInstanceID1 );
+			const FVertexID VertexID2 = MD->GetVertexInstanceVertex( VertexInstanceID2 );
 
 			const FVector DPosition1 = VertexPositions[ VertexID1 ] - VertexPositions[ VertexID0 ];
 			const FVector DPosition2 = VertexPositions[ VertexID2 ] - VertexPositions[ VertexID0 ];
 
-			const FVector2D DUV1 = VertexUVs.Get( Triangle.VertexInstanceID1, 0 ) - VertexUVs.Get( Triangle.VertexInstanceID0, 0 );
-			const FVector2D DUV2 = VertexUVs.Get( Triangle.VertexInstanceID2, 0 ) - VertexUVs.Get( Triangle.VertexInstanceID0, 0 );
+			const FVector2D DUV1 = VertexUVs.Get( VertexInstanceID1, 0 ) - VertexUVs.Get( VertexInstanceID0, 0 );
+			const FVector2D DUV2 = VertexUVs.Get( VertexInstanceID2, 0 ) - VertexUVs.Get( VertexInstanceID0, 0 );
 
 			// We have a left-handed coordinate system, but a counter-clockwise winding order
 			// Hence normal calculation has to take the triangle vectors cross product in reverse.
@@ -6120,7 +6001,7 @@ void UEditableMesh::GenerateTangentsAndNormals()
 
 	for( const FPolygonID PolygonID : PolygonsPendingNewTangentBasis )
 	{
-		VertexInstanceIDs.Append( GetMeshDescription()->GetPolygonPerimeterVertexInstances( PolygonID ) );
+		VertexInstanceIDs.Append( GetMeshDescription()->GetPolygonVertexInstances( PolygonID ) );
 	}
 
 	static TArray<FAttributesForVertexInstance> AttributesForVertexInstances;
@@ -6144,7 +6025,7 @@ void UEditableMesh::GenerateTangentsAndNormals()
 		// The vertex instance normal is computed as a sum of all connected polygons' normals, weighted by the angle they make with the vertex
 		for( const FPolygonID ConnectedPolygonID : AllConnectedPolygons )
 		{
-			const float Angle = GetPolygonCornerAngleForVertex( ConnectedPolygonID, VertexID );
+			const float Angle = GetMeshDescription()->GetPolygonCornerAngleForVertex( ConnectedPolygonID, VertexID );
 
 			Normal += PolygonNormals[ ConnectedPolygonID ] * Angle;
 
@@ -6196,7 +6077,7 @@ void UEditableMesh::FlipTangentsAndNormals()
 
 	for( const FPolygonID PolygonID : PolygonsPendingFlipTangentBasis )
 	{
-		VertexInstanceIDs.Append( GetMeshDescription()->GetPolygonPerimeterVertexInstances( PolygonID ) );
+		VertexInstanceIDs.Append( GetMeshDescription()->GetPolygonVertexInstances( PolygonID ) );
 	}
 
 	static TArray<FAttributesForVertexInstance> AttributesForVertexInstances;
@@ -6241,7 +6122,7 @@ void UEditableMesh::SplitPolygonalMesh(const FPlane& InPlane, TArray<FPolygonID>
     for (const auto& PolygonID : GetMeshDescription()->Polygons().GetElementIDs())
     {
         PolygonToEdgesMap.Add(PolygonID, {});
-        GetMeshDescription()->GetPolygonEdges(PolygonID, PolygonToEdgesMap[PolygonID]);
+        GetMeshDescription()->GetPolygonPerimeterEdges(PolygonID, PolygonToEdgesMap[PolygonID]);
     }
     
     // Find polygons that need to be split and if they don't need to add them to the appropriate list
@@ -6253,7 +6134,7 @@ void UEditableMesh::SplitPolygonalMesh(const FPlane& InPlane, TArray<FPolygonID>
         if (PotentialPolygonsToSplitSet.Contains(PolygonID))
         {
             TSet<FVertexInstanceID> VertexInstanceIDs;
-            VertexInstanceIDs.Append(GetMeshDescription()->GetPolygonPerimeterVertexInstances(PolygonID));
+            VertexInstanceIDs.Append(GetMeshDescription()->GetPolygonVertexInstances(PolygonID));
             int32 PosNeg = 0;
             TArray<bool> PosNegResults;
             for (const auto& VertexInstanceID : VertexInstanceIDs)
@@ -6299,7 +6180,7 @@ void UEditableMesh::SplitPolygonalMesh(const FPlane& InPlane, TArray<FPolygonID>
                 {
                     TriangulatedPolygons.Add(NewPolygonID);
                     PolygonToEdgesMap.Add(NewPolygonID, {});
-                    GetMeshDescription()->GetPolygonEdges(PolygonID, PolygonToEdgesMap[NewPolygonID]);
+                    GetMeshDescription()->GetPolygonPerimeterEdges(PolygonID, PolygonToEdgesMap[NewPolygonID]);
                 }
             }
             else
@@ -6658,12 +6539,13 @@ void UEditableMesh::SetEdgesHardnessAutomatically( const TArray<FEdgeID>& EdgeID
 }
 
 
+#if 0
 void UEditableMesh::SetEdgesVertices( const TArray<FVerticesForEdge>& VerticesForEdges )
 {
 	EM_ENTER( TEXT( "SetEdgesVertices: %s" ), *LogHelpers::ArrayToString( VerticesForEdges ) );
 
-	FEdgeArray& Edges = GetMeshDescription()->Edges();
-	FVertexArray& Vertices = GetMeshDescription()->Vertices();
+//	FEdgeArray& Edges = GetMeshDescription()->Edges();
+//	FVertexArray& Vertices = GetMeshDescription()->Vertices();
 
 	FSetEdgesVerticesChangeInput RevertInput;
 	RevertInput.VerticesForEdges.AddUninitialized( VerticesForEdges.Num() );
@@ -6674,36 +6556,21 @@ void UEditableMesh::SetEdgesVertices( const TArray<FVerticesForEdge>& VerticesFo
 	for( int32 EdgeNumber = 0; EdgeNumber < VerticesForEdges.Num(); ++EdgeNumber )
 	{
 		const FVerticesForEdge& VerticesForEdge = VerticesForEdges[ EdgeNumber ];
+		const FEdgeID EdgeID = VerticesForEdge.EdgeID;
 
 		// Save the backup
 		FVerticesForEdge& RevertVerticesForEdge = RevertInput.VerticesForEdges[ EdgeNumber ];
-		RevertVerticesForEdge.EdgeID = VerticesForEdge.EdgeID;
-		GetEdgeVertices( VerticesForEdge.EdgeID, /* Out */ RevertVerticesForEdge.NewVertexID0, /* Out */ RevertVerticesForEdge.NewVertexID1 );
+		RevertVerticesForEdge.EdgeID = EdgeID;
+		GetEdgeVertices( EdgeID, /* Out */ RevertVerticesForEdge.NewVertexID0, /* Out */ RevertVerticesForEdge.NewVertexID1 );
 
-		// Edit the edge
-		FMeshEdge& Edge = Edges[ VerticesForEdge.EdgeID ];
-		EdgeIDs[ EdgeNumber ] = VerticesForEdge.EdgeID;
-
-		for( uint32 EdgeVertexNumber = 0; EdgeVertexNumber < 2; ++EdgeVertexNumber )
-		{
-			// Disconnect the edge from its existing vertices
-			const FVertexID VertexID = Edge.VertexIDs[ EdgeVertexNumber ];
-			FMeshVertex& Vertex = Vertices[ VertexID ];
-			verify( Vertex.ConnectedEdgeIDs.RemoveSingleSwap( VerticesForEdge.EdgeID ) == 1 );	// Must have been already connected!
-		}
+		// Get current edge details
+		const FVertexID VertexID0 = GetMeshDescription()->GetEdgeVertex( EdgeID, 0 );
+		const FVertexID VertexID1 = GetMeshDescription()->GetEdgeVertex( EdgeID, 1 );
+		TArray<FPolygonID> ConnectedPolygons = GetMeshDescription()->GetEdgeConnectedPolygons( EdgeID );
 
-		Edge.VertexIDs[0] = VerticesForEdge.NewVertexID0;
-		Edge.VertexIDs[1] = VerticesForEdge.NewVertexID1;
-
-		// Connect the new vertices to the edge
-		for( uint32 EdgeVertexNumber = 0; EdgeVertexNumber < 2; ++EdgeVertexNumber )
-		{
-			const FVertexID VertexID = Edge.VertexIDs[ EdgeVertexNumber ];
-			FMeshVertex& Vertex = Vertices[ VertexID ];
-
-			check( !Vertex.ConnectedEdgeIDs.Contains( VerticesForEdge.EdgeID ) );	// Should not have already been connected
-			Vertex.ConnectedEdgeIDs.Add( VerticesForEdge.EdgeID );
-		}
+		// Renew edge
+		GetMeshDescription()->DeleteEdge( EdgeID );
+		GetMeshDescription()->CreateEdgeWithID( EdgeID, VerticesForEdge.NewVertexID0, VerticesForEdge.NewVertexID1, ConnectedPolygons );
 	}
 
 	// Give the adapter a chance to handle this
@@ -6847,6 +6714,7 @@ void UEditableMesh::RemovePolygonPerimeterVertices( const FPolygonID PolygonID,
 
 	EM_EXIT( TEXT( "RemovePolygonPerimeterVertices finished" ) );
 }
+#endif
 
 
 int32 UEditableMesh::FindPolygonPerimeterVertexNumberForVertex( const FPolygonID PolygonID, const FVertexID VertexID ) const
@@ -6940,13 +6808,9 @@ void UEditableMesh::TriangulatePolygons( const TArray<FPolygonID>& PolygonIDs, T
 			// We'll be replacing this polygon with it's triangulated counterpart polygons
 			PolygonsToDelete.Add( PolygonID );
 
-			// Figure out the triangulation for this polygon
-			static TArray<FMeshTriangle> Triangles;
-			ComputePolygonTriangulation( PolygonID, /* Out */ Triangles );
-
 			// Build polygons for each of the triangles that made up the original
 			{
-				for( const FMeshTriangle& Triangle : Triangles )
+				for( const FTriangleID& TriangleID : GetMeshDescription()->GetPolygonTriangleIDs( PolygonID ) )
 				{
 					PolygonsToCreate.Emplace();
 					FPolygonToCreate& PolygonToCreate = PolygonsToCreate.Last();
@@ -6959,7 +6823,7 @@ void UEditableMesh::TriangulatePolygons( const TArray<FPolygonID>& PolygonIDs, T
 						PolygonToCreate.PerimeterVertices.Emplace();
 						FVertexAndAttributes& PerimeterVertex = PolygonToCreate.PerimeterVertices.Last();
 
-						PerimeterVertex.VertexInstanceID = Triangle.GetVertexInstanceID( TriangleVertexNumber );
+						PerimeterVertex.VertexInstanceID = GetMeshDescription()->GetTriangleVertexInstance( TriangleID, TriangleVertexNumber );
 						PerimeterVertex.VertexID = FVertexID::Invalid;
 					}
 				}
@@ -7016,34 +6880,20 @@ void UEditableMesh::AssignPolygonsToPolygonGroups( const TArray<FPolygonGroupFor
 	static TArray<FPolygonGroupID> PolygonGroupsToDelete;
 	PolygonGroupsToDelete.Reset();
 
-	FPolygonArray& Polygons = GetMeshDescription()->Polygons();
-	FPolygonGroupArray& PolygonGroups = GetMeshDescription()->PolygonGroups();
-
 	for( const FPolygonGroupForPolygon& PolygonGroupForPolygon : PolygonGroupForPolygons )
 	{
 		const FPolygonID PolygonID = PolygonGroupForPolygon.PolygonID;
 		const FPolygonGroupID NewPolygonGroupID = PolygonGroupForPolygon.PolygonGroupID;
+		const FPolygonGroupID OldPolygonGroupID = GetMeshDescription()->GetPolygonPolygonGroup( PolygonID );
 
-		FMeshPolygon& Polygon = Polygons[ PolygonID ];
-		const FPolygonGroupID OldPolygonGroupID = Polygon.PolygonGroupID;
-
-		// Remove polygon reference from old group
-		FMeshPolygonGroup& OldPolygonGroup = PolygonGroups[ OldPolygonGroupID ];
-		verify( OldPolygonGroup.Polygons.RemoveSwap( PolygonID ) == 1 );
+		// Add polygon reference to new group
+		GetMeshDescription()->SetPolygonPolygonGroup( PolygonID, NewPolygonGroupID );
 
 		// If old group is now empty (and we're deleting orphans), add it to the list to delete
-		if( bDeleteOrphanedPolygonGroups && OldPolygonGroup.Polygons.Num() == 0 )
+		if( bDeleteOrphanedPolygonGroups && GetMeshDescription()->GetPolygonGroupPolygons( OldPolygonGroupID ).Num() == 0 )
 		{
 			PolygonGroupsToDelete.Add( OldPolygonGroupID );
 		}
-
-		// Add polygon reference to new group
-		FMeshPolygonGroup& NewPolygonGroup = PolygonGroups[ NewPolygonGroupID ];
-		check( !NewPolygonGroup.Polygons.Contains( PolygonID ) );
-		NewPolygonGroup.Polygons.Add( PolygonID );
-
-		// Change polygon group
-		Polygon.PolygonGroupID = NewPolygonGroupID;
 	}
 
 	// If there's any groups to delete, do it now
@@ -7216,7 +7066,7 @@ void UEditableMesh::WeldVertices( const TArray<FVertexID>& VertexIDsToWeld, FVer
 		static TArray<FEdgeID> PolygonEdgeIDs;
 		GetPolygonPerimeterVertices( ConnectedPolygonID, PolygonVertexIDs );
 		GetPolygonPerimeterEdges( ConnectedPolygonID, PolygonEdgeIDs );
-		const TArray<FVertexInstanceID>& PolygonVertexInstanceIDs = GetMeshDescription()->GetPolygonPerimeterVertexInstances( ConnectedPolygonID );
+		const TArray<FVertexInstanceID>& PolygonVertexInstanceIDs = GetMeshDescription()->GetPolygonVertexInstances( ConnectedPolygonID );
 
 		// Get the index range of perimeter vertices to be welded.
 		// This should definitely be valid, as any invalid welded poly will have caused early exit, above.
@@ -7431,7 +7281,7 @@ void UEditableMesh::TessellatePolygons( const TArray<FPolygonID>& PolygonIDs, co
 			static TArray<FVertexID> PerimeterVertexIDs;
 			GetPolygonPerimeterVertices( PolygonID, /* Out */ PerimeterVertexIDs );
 
-			const TArray<FVertexInstanceID>& PerimeterVertexInstanceIDs = GetMeshDescription()->GetPolygonPerimeterVertexInstances( PolygonID );
+			const TArray<FVertexInstanceID>& PerimeterVertexInstanceIDs = GetMeshDescription()->GetPolygonVertexInstances( PolygonID );
 
 			const int32 PerimeterVertexCount = PerimeterEdgeCount;
 			const int32 OriginalPerimeterEdgeCount = PerimeterEdgeCount / 2;
@@ -7481,12 +7331,13 @@ void UEditableMesh::TessellatePolygons( const TArray<FPolygonID>& PolygonIDs, co
 						{
 							const FVector CenterVertexPosition = VertexPositions[ PolygonCenterVertexID ];
 
-							FMeshTriangle Triangle;
 							FVector TriangleVertexWeights;
-							if( ComputeBarycentricWeightForPointOnPolygon( PolygonID, CenterVertexPosition, /* Out */ Triangle, /* Out */ TriangleVertexWeights ) )
+							const FTriangleID TriangleID = ComputeBarycentricWeightForPointOnPolygon( PolygonID, CenterVertexPosition, /* Out */ TriangleVertexWeights );
+							if( TriangleID != FTriangleID::Invalid )
 							{
+								TArrayView<const FVertexInstanceID> TriangleVertexInstanceIDs = GetMeshDescription()->GetTriangleVertexInstances( TriangleID );
 								InterpAllAttributes( VertexAndAttributes.PolygonVertexAttributes, GetMeshDescription()->VertexInstanceAttributes(),
-									Triangle.VertexInstanceID0, Triangle.VertexInstanceID1, Triangle.VertexInstanceID2, TriangleVertexWeights );
+									TriangleVertexInstanceIDs[0], TriangleVertexInstanceIDs[1], TriangleVertexInstanceIDs[2], TriangleVertexWeights );
 							}
 						}
 					}
@@ -7502,7 +7353,7 @@ void UEditableMesh::TessellatePolygons( const TArray<FPolygonID>& PolygonIDs, co
 		{
 			if( TriangleTessellationMode == ETriangleTessellationMode::ThreeTriangles )
 			{
-				const TArray<FVertexInstanceID>& PerimeterVertexInstanceIDs = GetMeshDescription()->GetPolygonPerimeterVertexInstances( PolygonID );
+				const TArray<FVertexInstanceID>& PerimeterVertexInstanceIDs = GetMeshDescription()->GetPolygonVertexInstances( PolygonID );
 				check( PerimeterVertexInstanceIDs.Num() == 3 );
 
 				// Define the three new triangles for the original tessellated triangle
@@ -7528,12 +7379,13 @@ void UEditableMesh::TessellatePolygons( const TArray<FPolygonID>& PolygonIDs, co
 					{
 						const FVector CenterVertexPosition = VertexPositions[ PolygonCenterVertexID ];
 
-						FMeshTriangle Triangle;
 						FVector TriangleVertexWeights;
-						if( ComputeBarycentricWeightForPointOnPolygon( PolygonID, CenterVertexPosition, /* Out */ Triangle, /* Out */ TriangleVertexWeights ) )
+						const FTriangleID TriangleID = ComputeBarycentricWeightForPointOnPolygon( PolygonID, CenterVertexPosition, /* Out */ TriangleVertexWeights );
+						if( TriangleID != FTriangleID::Invalid )
 						{
+							TArrayView<const FVertexInstanceID> TriangleVertexInstanceIDs = GetMeshDescription()->GetTriangleVertexInstances( TriangleID );
 							InterpAllAttributes( PolygonToCreate.PerimeterVertices[ 2 ].PolygonVertexAttributes, GetMeshDescription()->VertexInstanceAttributes(),
-								Triangle.VertexInstanceID0, Triangle.VertexInstanceID1, Triangle.VertexInstanceID2, TriangleVertexWeights );
+								TriangleVertexInstanceIDs[ 0 ], TriangleVertexInstanceIDs[ 1 ], TriangleVertexInstanceIDs[ 2 ], TriangleVertexWeights );
 						}
 					}
 				}
@@ -7544,7 +7396,7 @@ void UEditableMesh::TessellatePolygons( const TArray<FPolygonID>& PolygonIDs, co
 				// the center, connecting the three new vertices that we created between each original edge.  The
 				// other three triangles will go in the corners of the original triangle.
 
-				const TArray<FVertexInstanceID>& PerimeterVertexInstanceIDs = GetMeshDescription()->GetPolygonPerimeterVertexInstances( PolygonID );
+				const TArray<FVertexInstanceID>& PerimeterVertexInstanceIDs = GetMeshDescription()->GetPolygonVertexInstances( PolygonID );
 				check( PerimeterVertexInstanceIDs.Num() == 6 );	// We split the triangle's 3 edges earlier, so we must have six edges now
 
 				// Define the new center triangle
@@ -7848,8 +7700,8 @@ void UEditableMesh::QuadrangulatePolygonGroup( const FPolygonGroupID PolygonGrou
 								const int32 PrevAdjacentPerimeterEdgeIndex = ( AdjacentPerimeterEdgeIndex + 2 ) % 3;
 								const int32 NextAdjacentPerimeterEdgeIndex = ( AdjacentPerimeterEdgeIndex + 1 ) % 3;
 
-								const TArray<FVertexInstanceID>& PerimeterVertexInstanceIDs = GetMeshDescription()->GetPolygonPerimeterVertexInstances( PolygonID );
-								const TArray<FVertexInstanceID>& AdjacentPerimeterVertexInstanceIDs = GetMeshDescription()->GetPolygonPerimeterVertexInstances( AdjacentPolygonID );
+								const TArray<FVertexInstanceID>& PerimeterVertexInstanceIDs = GetMeshDescription()->GetPolygonVertexInstances( PolygonID );
+								const TArray<FVertexInstanceID>& AdjacentPerimeterVertexInstanceIDs = GetMeshDescription()->GetPolygonVertexInstances( AdjacentPolygonID );
 
 								const FVertexInstanceID VertexInstanceStart1 = PerimeterVertexInstanceIDs[ PerimeterEdgeIndex ];
 								const FVertexInstanceID VertexInstanceStart2 = AdjacentPerimeterVertexInstanceIDs[ AdjacentPerimeterEdgeIndex ];
diff --git a/Engine/Plugins/Runtime/EditableMesh/Source/EditableMesh/EditableMeshChanges.cpp b/Engine/Plugins/Runtime/EditableMesh/Source/EditableMesh/EditableMeshChanges.cpp
--- a/Engine/Plugins/Runtime/EditableMesh/Source/EditableMesh/EditableMeshChanges.cpp
+++ b/Engine/Plugins/Runtime/EditableMesh/Source/EditableMesh/EditableMeshChanges.cpp
@@ -270,6 +270,7 @@ FString FChangePolygonsVertexInstancesChange::ToString() const
 }
 
 
+#if 0
 TUniquePtr<FChange> FSetEdgesVerticesChange::Execute( UObject* Object )
 {
 	UEditableMesh* EditableMesh = CastChecked<UEditableMesh>( Object );
@@ -327,6 +328,7 @@ FString FRemovePolygonPerimeterVerticesChange::ToString() const
 		Input.FirstVertexNumberToRemove,
 		Input.NumVerticesToRemove );
 }
+#endif
 
 
 TUniquePtr<FChange> FStartOrEndModificationChange::Execute( UObject* Object )
diff --git a/Engine/Plugins/Runtime/EditableMesh/Source/EditableMesh/EditableStaticMeshAdapter.cpp b/Engine/Plugins/Runtime/EditableMesh/Source/EditableMesh/EditableStaticMeshAdapter.cpp
--- a/Engine/Plugins/Runtime/EditableMesh/Source/EditableMesh/EditableStaticMeshAdapter.cpp
+++ b/Engine/Plugins/Runtime/EditableMesh/Source/EditableMesh/EditableStaticMeshAdapter.cpp
@@ -12,9 +12,6 @@
 #include "MeshAttributes.h"
 
 
-const FTriangleID FTriangleID::Invalid( TNumericLimits<uint32>::Max() );
-
-
 UEditableStaticMeshAdapter::UEditableStaticMeshAdapter()
 	: StaticMesh( nullptr ),
 	  StaticMeshLODIndex( 0 ),
@@ -25,7 +22,7 @@ UEditableStaticMeshAdapter::UEditableStaticMeshAdapter()
 }
 
 
-inline void UEditableStaticMeshAdapter::EnsureIndexBufferIs32Bit()
+void UEditableStaticMeshAdapter::EnsureIndexBufferIs32Bit()
 {
 	FStaticMeshLODResources& StaticMeshLOD = GetStaticMeshLOD();
 	if( !StaticMeshLOD.IndexBuffer.Is32Bit() )
@@ -38,27 +35,6 @@ inline void UEditableStaticMeshAdapter::EnsureIndexBufferIs32Bit()
 }
 
 
-inline void UEditableStaticMeshAdapter::UpdateIndexBufferFormatIfNeeded( const TArray<FMeshTriangle>& Triangles )
-{
-	const FStaticMeshLODResources& StaticMeshLOD = GetStaticMeshLOD();
-	if( !StaticMeshLOD.IndexBuffer.Is32Bit() )
-	{
-		for( const FMeshTriangle& Triangle : Triangles )
-		{
-			for( int32 TriangleVertexNumber = 0; TriangleVertexNumber < 3; ++TriangleVertexNumber )
-			{
-				const FVertexInstanceID VertexInstanceID = Triangle.GetVertexInstanceID( TriangleVertexNumber );
-				if( VertexInstanceID.GetValue() > TNumericLimits<uint16>::Max() )
-				{
-					EnsureIndexBufferIs32Bit();
-					return;
-				}
-			}
-		}
-	}
-}
-
-
 void UEditableStaticMeshAdapter::InitEditableStaticMesh( UEditableMesh* EditableMesh, class UPrimitiveComponent& Component, const FEditableMeshSubMeshAddress& InitSubMeshAddress )
 {
 	EditableMesh->SetSubMeshAddress( InitSubMeshAddress );
@@ -318,6 +294,7 @@ void UEditableStaticMeshAdapter::InitEditableStaticMesh( UEditableMesh* Editable
 						{
 							// Static meshes only support triangles, so there's no need to triangulate anything yet.  We'll make both
 							// a triangle and a polygon here.
+							// @todo: get rid of rendering triangles completely
 							const FTriangleID NewTriangleID = FTriangleID( SectionTriangleIndex );
 
 							NewRenderingPolygonGroup.Triangles.Insert( NewTriangleID );
@@ -336,9 +313,6 @@ void UEditableStaticMeshAdapter::InitEditableStaticMesh( UEditableMesh* Editable
 							FRenderingPolygon& NewRenderingPolygon = RenderingPolygons[ NewPolygonID ];
 							NewRenderingPolygon.PolygonGroupID = NewPolygonGroupID;
 							NewRenderingPolygon.TriangulatedPolygonTriangleIndices.Add( NewTriangleID );
-
-							// Add triangle to polygon triangulation array
-							MeshDescription->GetPolygonTriangles( NewPolygonID ).Add( NewTriangle );
 						}
 						else
 						{
@@ -448,14 +422,19 @@ void UEditableStaticMeshAdapter::InitializeFromEditableMesh( const UEditableMesh
 		FRenderingPolygon& RenderingPolygon = RenderingPolygons[ PolygonID ];
 		RenderingPolygon.PolygonGroupID = PolygonGroupID;
 
-		const TArray<FMeshTriangle>& Triangles = MeshDescription->GetPolygonTriangles( PolygonID );
-		for( const FMeshTriangle& Triangle : Triangles )
+		const TArray<FTriangleID>& TriangleIDs = MeshDescription->GetPolygonTriangleIDs( PolygonID );
+		for( const FTriangleID TriangleID : TriangleIDs )
 		{
-			const FTriangleID TriangleID = RenderingPolygonGroup.Triangles.Add( Triangle );
-			RenderingPolygon.TriangulatedPolygonTriangleIndices.Add( TriangleID );
+			FMeshTriangle Triangle;
+			for( int32 TriIndex = 0; TriIndex < 3; ++TriIndex )
+			{
+				Triangle.SetVertexInstanceID( TriIndex, MeshDescription->GetTriangleVertexInstance( TriangleID, TriIndex ) );
+			}
+			const FTriangleID RenderingTriangleID = RenderingPolygonGroup.Triangles.Add( Triangle );
+			RenderingPolygon.TriangulatedPolygonTriangleIndices.Add( RenderingTriangleID );
 		}
 
-		RenderingPolygonGroup.MaxTriangles += Triangles.Num();
+		RenderingPolygonGroup.MaxTriangles += TriangleIDs.Num();
 	}
 
 }
@@ -1319,12 +1298,12 @@ void UEditableStaticMeshAdapter::OnRetriangulatePolygons( const UEditableMesh* E
 		FRenderingPolygon& RenderingPolygon = RenderingPolygons[ PolygonID ];
 		const FPolygonGroupID PolygonGroupID = RenderingPolygon.PolygonGroupID;
 		FRenderingPolygonGroup& RenderingPolygonGroup = RenderingPolygonGroups[ PolygonGroupID ];
-		const TArray<FMeshTriangle>& Triangles = MeshDescription->GetPolygonTriangles( PolygonID );
+		const TArray<FTriangleID>& TriangleIDs = MeshDescription->GetPolygonTriangleIDs( PolygonID );
 
 		// Check to see whether the index buffer needs to be updated
 		bool bNeedsUpdatedTriangles = false;
 		{
-			if( RenderingPolygon.TriangulatedPolygonTriangleIndices.Num() != Triangles.Num() )
+			if( RenderingPolygon.TriangulatedPolygonTriangleIndices.Num() != TriangleIDs.Num() )
 			{
 				// Triangle count has changed, so we definitely need new triangles!
 				bNeedsUpdatedTriangles = true;
@@ -1332,14 +1311,14 @@ void UEditableStaticMeshAdapter::OnRetriangulatePolygons( const UEditableMesh* E
 			else
 			{
 				// See if the triangulation has changed even if the number of triangles is the same
-				for( int32 Index = 0; Index < Triangles.Num(); ++Index )
+				for( int32 Index = 0; Index < TriangleIDs.Num(); ++Index )
 				{
 					const FMeshTriangle& OldTriangle = RenderingPolygonGroup.Triangles[ RenderingPolygon.TriangulatedPolygonTriangleIndices[ Index ] ];
-					const FMeshTriangle& NewTriangle = Triangles[ Index ];
+					const FTriangleID NewTriangleID = TriangleIDs[ Index ];
 
-					if( OldTriangle.VertexInstanceID0 != NewTriangle.VertexInstanceID0 ||
-						OldTriangle.VertexInstanceID1 != NewTriangle.VertexInstanceID1 ||
-						OldTriangle.VertexInstanceID2 != NewTriangle.VertexInstanceID2 )
+					if( OldTriangle.GetVertexInstanceID( 0 ) != MeshDescription->GetTriangleVertexInstance( NewTriangleID, 0 ) ||
+						OldTriangle.GetVertexInstanceID( 1 ) != MeshDescription->GetTriangleVertexInstance( NewTriangleID, 1 ) ||
+						OldTriangle.GetVertexInstanceID( 2 ) != MeshDescription->GetTriangleVertexInstance( NewTriangleID, 2 ) )
 					{
 						bNeedsUpdatedTriangles = true;
 						break;
@@ -1361,7 +1340,7 @@ void UEditableStaticMeshAdapter::OnRetriangulatePolygons( const UEditableMesh* E
 			// Add new triangles
 			{
 				// This is the number of triangles we are about to add
-				const int32 NumNewTriangles = Triangles.Num();
+				const int32 NumNewTriangles = TriangleIDs.Num();
 
 				// This is the number of entries currently unused in the Triangles sparse array
 				const int32 NumFreeTriangles = RenderingPolygonGroup.Triangles.GetArraySize() - RenderingPolygonGroup.Triangles.Num();
@@ -1384,6 +1363,7 @@ void UEditableStaticMeshAdapter::OnRetriangulatePolygons( const UEditableMesh* E
 
 				// Create empty triangles for all of the new triangles we need, and keep track of their triangle indices
 				static TArray<FTriangleID> NewTriangleIDs;
+				bool bUses32BitIndices = false;
 				{
 					NewTriangleIDs.SetNumUninitialized( NumNewTriangles, false );
 
@@ -1395,7 +1375,11 @@ void UEditableStaticMeshAdapter::OnRetriangulatePolygons( const UEditableMesh* E
 						FMeshTriangle& NewTriangle = RenderingPolygonGroup.Triangles[ NewTriangleID ];
 						for( int32 TriangleVertexNumber = 0; TriangleVertexNumber < 3; ++TriangleVertexNumber )
 						{
-							const FVertexInstanceID VertexInstanceID = Triangles[ TriangleToAddNumber ].GetVertexInstanceID( TriangleVertexNumber );
+							const FVertexInstanceID VertexInstanceID = MeshDescription->GetTriangleVertexInstance( TriangleIDs[ TriangleToAddNumber ], TriangleVertexNumber );
+							if( VertexInstanceID.GetValue() > TNumericLimits<uint16>::Max() )
+							{
+								bUses32BitIndices = true;
+							}
 							NewTriangle.SetVertexInstanceID( TriangleVertexNumber, VertexInstanceID );
 							MinVertexIndex = FMath::Min( MinVertexIndex, VertexInstanceID.GetValue() );
 							MaxVertexIndex = FMath::Max( MaxVertexIndex, VertexInstanceID.GetValue() );
@@ -1406,9 +1390,9 @@ void UEditableStaticMeshAdapter::OnRetriangulatePolygons( const UEditableMesh* E
 				}
 
 				// Update the index buffer format if the index range exceeds 16 bit values.
-				if( !EditableMesh->IsPreviewingSubdivisions() )
+				if( !EditableMesh->IsPreviewingSubdivisions() && bUses32BitIndices )
 				{
-					UpdateIndexBufferFormatIfNeeded( Triangles );
+					EnsureIndexBufferIs32Bit();
 				}
 
 				// If we need more space in the index buffer for this section, allocate it here
diff --git a/Engine/Plugins/Runtime/EditableMesh/Source/EditableMesh/EditableStaticMeshAdapter.h b/Engine/Plugins/Runtime/EditableMesh/Source/EditableMesh/EditableStaticMeshAdapter.h
--- a/Engine/Plugins/Runtime/EditableMesh/Source/EditableMesh/EditableStaticMeshAdapter.h
+++ b/Engine/Plugins/Runtime/EditableMesh/Source/EditableMesh/EditableStaticMeshAdapter.h
@@ -9,35 +9,6 @@
 #include "EditableStaticMeshAdapter.generated.h"
 
 
-USTRUCT()
-struct FTriangleID : public FElementID
-{
-	GENERATED_BODY()
-
-	FTriangleID()
-	{
-	}
-
-	explicit FTriangleID( const FElementID InitElementID )
-		: FElementID( InitElementID.GetValue() )
-	{
-	}
-
-	explicit FTriangleID( const uint32 InitIDValue )
-		: FElementID( InitIDValue )
-	{
-	}
-
-	FORCEINLINE friend uint32 GetTypeHash( const FTriangleID& Other )
-	{
-		return GetTypeHash( Other.IDValue );
-	}
-
-	/** Invalid triangle ID */
-	static const FTriangleID Invalid;
-};
-
-
 USTRUCT()
 struct FRenderingPolygon
 {
@@ -180,9 +151,6 @@ private:
 	/** Makes sure our mesh's index buffer is 32-bit, converting if needed */
 	void EnsureIndexBufferIs32Bit();
 
-	/** If any of the specified triangles contain vertex indices greater than 0xffff, updates our mesh's index buffer to be 32-bit if it isn't already */
-	void UpdateIndexBufferFormatIfNeeded( const TArray<FMeshTriangle>& Triangles );
-
 	/** Rebuilds bounds */
 	void UpdateBounds( const UEditableMesh* EditableMesh, const bool bShouldRecompute );
 
diff --git a/Engine/Plugins/Runtime/EditableMesh/Source/EditableMesh/GeometryHitTest.cpp b/Engine/Plugins/Runtime/EditableMesh/Source/EditableMesh/GeometryHitTest.cpp
--- a/Engine/Plugins/Runtime/EditableMesh/Source/EditableMesh/GeometryHitTest.cpp
+++ b/Engine/Plugins/Runtime/EditableMesh/Source/EditableMesh/GeometryHitTest.cpp
@@ -94,18 +94,13 @@ FEditableMeshElementAddress FGeometryTests::QueryElement(const UEditableMesh& Ed
 	// Check polygons first; this is so we always impose a closest hit location at the poly before checking other elements, so anything behind is occluded
 	for (const FPolygonID PolygonID : FrontFacingPolygons)
 	{
-		static TArray<FVertexID> MeshVertexIDs;
-		MeshVertexIDs.Reset();
-		EditableMesh.GetPolygonPerimeterVertices(PolygonID, /* Out */ MeshVertexIDs);
-
-		const uint32 PolygonTriangleCount = EditableMesh.GetPolygonTriangulatedTriangleCount(PolygonID);
-		for (uint32 PolygonTriangleNumber = 0; PolygonTriangleNumber < PolygonTriangleCount; ++PolygonTriangleNumber)
+		for (const FTriangleID TriangleID : MeshDescription->GetPolygonTriangleIDs(PolygonID))
 		{
 			FVector TriangleVertexPositions[3];
 			for (uint32 TriangleVertexNumber = 0; TriangleVertexNumber < 3; ++TriangleVertexNumber)
 			{
-				const FVertexInstanceID VertexInstanceID = EditableMesh.GetPolygonTriangulatedTriangle(PolygonID, PolygonTriangleNumber).GetVertexInstanceID(TriangleVertexNumber);
-				const FVertexID VertexID = EditableMesh.GetVertexInstanceVertex(VertexInstanceID);
+				const FVertexInstanceID VertexInstanceID = MeshDescription->GetTriangleVertexInstance(TriangleID, TriangleVertexNumber);
+				const FVertexID VertexID = MeshDescription->GetVertexInstanceVertex(VertexInstanceID);
 				TriangleVertexPositions[TriangleVertexNumber] = VertexPositions[VertexID];
 			}
 
diff --git a/Engine/Plugins/Runtime/EditableMesh/Source/EditableMesh/Public/EditableMesh.h b/Engine/Plugins/Runtime/EditableMesh/Source/EditableMesh/Public/EditableMesh.h
--- a/Engine/Plugins/Runtime/EditableMesh/Source/EditableMesh/Public/EditableMesh.h
+++ b/Engine/Plugins/Runtime/EditableMesh/Source/EditableMesh/Public/EditableMesh.h
@@ -16,7 +16,7 @@
 
 class UEditableMeshAdapter;
 
-#define EDITABLE_MESH_DEBUG_LOG
+//#define EDITABLE_MESH_DEBUG_LOG
 #ifdef EDITABLE_MESH_DEBUG_LOG
 EDITABLEMESH_API DECLARE_LOG_CATEGORY_EXTERN( LogEditableMesh, All, All );
 #else
@@ -218,7 +218,7 @@ public:
 
 	/** Returns the indexed triangle of the triangulated polygon */
 	UFUNCTION( BlueprintPure, Category="Editable Mesh" )
-	FMeshTriangle GetPolygonTriangulatedTriangle( const FPolygonID PolygonID, int32 PolygonTriangleNumber ) const;
+	FTriangleID GetPolygonTriangulatedTriangle( const FPolygonID PolygonID, int32 PolygonTriangleNumber ) const;
 
 protected:
 
@@ -231,15 +231,13 @@ protected:
 	void SetPolygonAttribute( const FPolygonID PolygonID, const FMeshElementAttributeData& Attribute );
 	void SetPolygonGroupAttribute( const FPolygonGroupID PolygonGroupID, const FMeshElementAttributeData& Attribute );
 	FVertexInstanceID CreateVertexInstanceForContourVertex( const FVertexAndAttributes& ContourVertex, const FPolygonID PolygonID );
-	void CreatePolygonContour( const TArray<FVertexAndAttributes>& Contour, TArray<FVertexInstanceID>& OutVertexInstanceIDs );
-	void BackupPolygonContour( const FMeshPolygonContour& Contour, TArray<FVertexAndAttributes>& OutVerticesAndAttributes );
+	void CreatePolygonContour( const TArray<FVertexAndAttributes>& Contour, TArray<FEdgeID>& OutEdgeIDs, TArray<FVertexInstanceID>& OutVertexInstanceIDs );
 	void GetConnectedSoftEdges( const FVertexID VertexID, TArray<FEdgeID>& OutConnectedSoftEdges ) const;
 	void GetVertexConnectedPolygonsInSameSoftEdgedGroup( const FVertexID VertexInstanceID, const FPolygonID PolygonID, TArray<FPolygonID>& OutPolygonIDs ) const;
 	FVertexInstanceID GetVertexInstanceInPolygonForVertex( const FPolygonID PolygonID, const FVertexID VertexID ) const;
-	void SetPolygonContourVertexAttributes( FMeshPolygonContour& Contour, const FPolygonID PolygonID, const TArray<FMeshElementAttributeList>& AttributeLists );
+	void SetPolygonContourVertexAttributes( const FPolygonID PolygonID, const TArray<FMeshElementAttributeList>& AttributeLists );
 	void SplitVertexInstanceInPolygons( const FVertexInstanceID VertexInstanceID, const TArray<FPolygonID>& PolygonIDs );
 	void ReplaceVertexInstanceInPolygons( const FVertexInstanceID OldVertexInstanceID, const FVertexInstanceID NewVertexInstanceID, const TArray<FPolygonID>& PolygonIDs );
-	float GetPolygonCornerAngleForVertex( const FPolygonID PolygonID, const FVertexID VertexID ) const;
 	void SplitVerticesIfNecessary( const TArray<FVertexID>& VerticesToSplit );
 	void GetPolygonsInSameSoftEdgedGroupAsPolygon( const FPolygonID PolygonID, const TArray<FPolygonID>& PolygonIDsToCheck, const TArray<FEdgeID>& SoftEdgeIDs, TArray<FPolygonID>& OutPolygonIDs ) const;
 
@@ -410,8 +408,6 @@ public:
 	UFUNCTION( BlueprintPure, Category="Editable Mesh" ) FPlane ComputePolygonPlane( const FPolygonID PolygonID ) const;
 	UFUNCTION( BlueprintPure, Category="Editable Mesh" ) FVector ComputePolygonNormal( const FPolygonID PolygonID ) const;
 	UFUNCTION( BlueprintPure, Category="Editable Mesh" ) const FSubdivisionLimitData& GetSubdivisionLimitData() const;
-	UFUNCTION( BlueprintPure, Category="Editable Mesh" ) void ComputePolygonTriangulation( const FPolygonID PolygonID, TArray<FMeshTriangle>& OutTriangles ) const;
-	UFUNCTION( BlueprintPure, Category="Editable Mesh" ) bool ComputeBarycentricWeightForPointOnPolygon( const FPolygonID PolygonID, const FVector PointOnPolygon, FMeshTriangle& OutTriangle, FVector& OutTriangleVertexWeights ) const;
 	UFUNCTION( BlueprintPure, Category="Editable Mesh" ) void ComputePolygonsSharedEdges( const TArray<FPolygonID>& PolygonIDs, TArray<FEdgeID>& OutSharedEdgeIDs ) const;
 	UFUNCTION( BlueprintPure, Category="Editable Mesh" ) void FindPolygonLoop( const FEdgeID EdgeID, TArray<FEdgeID>& OutEdgeLoopEdgeIDs, TArray<FEdgeID>& OutFlippedEdgeIDs, TArray<FEdgeID>& OutReversedEdgeIDPathToTake, TArray<FPolygonID>& OutPolygonIDsToSplit ) const;
 	UFUNCTION( BlueprintPure, Category="Editable Mesh" ) void SearchSpatialDatabaseForPolygonsPotentiallyIntersectingLineSegment( const FVector LineSegmentStart, const FVector LineSegmentEnd, TArray<FPolygonID>& OutPolygons ) const;
@@ -452,9 +448,9 @@ public:
 	UFUNCTION( BlueprintCallable, Category="Editable Mesh" ) void SetEdgesCreaseSharpness( const TArray<FEdgeID>& EdgeIDs, const TArray<float>& EdgesNewCreaseSharpness );
 	UFUNCTION( BlueprintCallable, Category="Editable Mesh" ) void SetEdgesHardness( const TArray<FEdgeID>& EdgeIDs, const TArray<bool>& EdgesNewIsHard );
 	UFUNCTION( BlueprintCallable, Category="Editable Mesh" ) void SetEdgesHardnessAutomatically( const TArray<FEdgeID>& EdgeIDs, const float MaxDotProductForSoftEdge );	// @todo mesheditor: Not used for anything yet.  Remove it?  Use it during import/convert?
-	UFUNCTION( BlueprintCallable, Category="Editable Mesh" ) void SetEdgesVertices( const TArray<FVerticesForEdge>& VerticesForEdges );
-	UFUNCTION( BlueprintCallable, Category="Editable Mesh" ) void InsertPolygonPerimeterVertices( const FPolygonID PolygonID, const int32 InsertBeforeVertexNumber, const TArray<FVertexAndAttributes>& VerticesToInsert );
-	UFUNCTION( BlueprintCallable, Category="Editable Mesh" ) void RemovePolygonPerimeterVertices( const FPolygonID PolygonID, const int32 FirstVertexNumberToRemove, const int32 NumVerticesToRemove, const bool bDeleteOrphanedVertexInstances );
+//	UFUNCTION( BlueprintCallable, Category="Editable Mesh" ) void SetEdgesVertices( const TArray<FVerticesForEdge>& VerticesForEdges );
+//	UFUNCTION( BlueprintCallable, Category="Editable Mesh" ) void InsertPolygonPerimeterVertices( const FPolygonID PolygonID, const int32 InsertBeforeVertexNumber, const TArray<FVertexAndAttributes>& VerticesToInsert );
+//	UFUNCTION( BlueprintCallable, Category="Editable Mesh" ) void RemovePolygonPerimeterVertices( const FPolygonID PolygonID, const int32 FirstVertexNumberToRemove, const int32 NumVerticesToRemove, const bool bDeleteOrphanedVertexInstances );
 	UFUNCTION( BlueprintCallable, Category="Editable Mesh" ) void FlipPolygons( const TArray<FPolygonID>& PolygonIDs );
 	UFUNCTION( BlueprintCallable, Category="Editable Mesh" ) void TriangulatePolygons( const TArray<FPolygonID>& PolygonIDs, TArray<FPolygonID>& OutNewTrianglePolygons );
 	UFUNCTION( BlueprintCallable, Category="Editable Mesh" ) void CreatePolygonGroups( const TArray<FPolygonGroupToCreate>& PolygonGroupsToCreate, TArray<FPolygonGroupID>& OutNewPolygonGroupIDs );
@@ -468,6 +464,7 @@ public:
 	UFUNCTION( BlueprintCallable, Category="Editable Mesh" ) void SplitPolygonalMesh( const FPlane& InPlane, TArray<FPolygonID>& PolygonIDs1, TArray<FPolygonID>& PolygonIDs2, TArray<FEdgeID>& BoundaryIDs);
 
 	void GeometryHitTest(const FHitParamsIn& InParams, FHitParamsOut& OutParams);
+	FTriangleID ComputeBarycentricWeightForPointOnPolygon( const FPolygonID PolygonID, const FVector PointOnPolygon, FVector& OutTriangleVertexWeights ) const;
 
 protected:
 
