Fixed crash when deleting sections in StaticMeshEditor.

diff --git a/Engine/Plugins/Editor/MeshEditor/Source/PolygonModeling/DeleteMeshElement.cpp b/Engine/Plugins/Editor/MeshEditor/Source/PolygonModeling/DeleteMeshElement.cpp
--- a/Engine/Plugins/Editor/MeshEditor/Source/PolygonModeling/DeleteMeshElement.cpp
+++ b/Engine/Plugins/Editor/MeshEditor/Source/PolygonModeling/DeleteMeshElement.cpp
@@ -53,7 +53,7 @@ void UDeleteMeshElementCommand::Execute( IMeshEditorModeEditingContract& MeshEdi
 			const bool bDeleteOrphanedEdges = true;
 			const bool bDeleteOrphanedVertices = true;
 			const bool bDeleteOrphanedVertexInstances = true;
-			const bool bDeleteEmptySections = true;
+			const bool bDeleteEmptySections = false;
 
 			// If we deleted the same polygon on multiple selected instances of the same mesh, the polygon could already have been deleted
 			// by the time we get here
diff --git a/Engine/Plugins/Runtime/EditableMesh/Source/EditableMesh/EditableStaticMeshAdapter.cpp b/Engine/Plugins/Runtime/EditableMesh/Source/EditableMesh/EditableStaticMeshAdapter.cpp
--- a/Engine/Plugins/Runtime/EditableMesh/Source/EditableMesh/EditableStaticMeshAdapter.cpp
+++ b/Engine/Plugins/Runtime/EditableMesh/Source/EditableMesh/EditableStaticMeshAdapter.cpp
@@ -624,7 +624,6 @@ void UEditableStaticMeshAdapter::OnRebuildRenderMesh( const UEditableMesh* Edita
 
 			StaticMeshSection.FirstIndex = IndexBuffer.Num();
 			StaticMeshSection.NumTriangles = RenderingPolygonGroup.Triangles.GetArraySize();
-			check( RenderingPolygonGroup.Triangles.GetArraySize() <= RenderingPolygonGroup.MaxTriangles );
 
 			const int32 MaterialIndex = StaticMesh->GetMaterialIndexFromImportedMaterialSlotName( PolygonGroupImportedMaterialSlotNames[ PolygonGroupID ] );
 			//check( MaterialIndex != INDEX_NONE );
@@ -682,6 +681,10 @@ void UEditableStaticMeshAdapter::OnRebuildRenderMesh( const UEditableMesh* Edita
 				// No triangles in this section
 				StaticMeshSection.MinVertexIndex = 0;
 				StaticMeshSection.MaxVertexIndex = 0;
+				StaticMeshSection.NumTriangles = 0;
+
+				// Setting this to 0 allows the adapter to allocate more index buffer in OnRetriangulatePolygons (for example, when undoing delete polygons)
+				RenderingPolygonGroup.MaxTriangles = 0;
 			}
 		}
 	}
@@ -1405,6 +1408,12 @@ void UEditableStaticMeshAdapter::OnRetriangulatePolygons( const UEditableMesh* E
 					const int32 MaxTriangles = RenderingPolygonGroup.MaxTriangles;
 					RenderingPolygonGroup.MaxTriangles += NumExtraTriangles;
 
+					// Update the number of triangles in the section when the number of triangles change in the rendering polygon group
+					if (RenderingSection.NumTriangles <  (uint32) RenderingPolygonGroup.MaxTriangles)
+					{
+						RenderingSection.NumTriangles = RenderingPolygonGroup.MaxTriangles;
+					}
+
 					if( !EditableMesh->IsPreviewingSubdivisions() )
 					{
 						const uint32 FirstIndex = RenderingSection.FirstIndex;
