Fixed crash when assigning material to subset of static mesh in staticmesh editor

#rnx

diff --git a/Engine/Plugins/Runtime/EditableMesh/Source/EditableMesh/EditableStaticMeshAdapter.cpp b/Engine/Plugins/Runtime/EditableMesh/Source/EditableMesh/EditableStaticMeshAdapter.cpp
--- a/Engine/Plugins/Runtime/EditableMesh/Source/EditableMesh/EditableStaticMeshAdapter.cpp
+++ b/Engine/Plugins/Runtime/EditableMesh/Source/EditableMesh/EditableStaticMeshAdapter.cpp
@@ -1408,12 +1408,6 @@ void UEditableStaticMeshAdapter::OnRetriangulatePolygons( const UEditableMesh* E
 					const int32 MaxTriangles = RenderingPolygonGroup.MaxTriangles;
 					RenderingPolygonGroup.MaxTriangles += NumExtraTriangles;
 
-					// Update the number of triangles in the section when the number of triangles change in the rendering polygon group
-					if (RenderingSection.NumTriangles <  (uint32) RenderingPolygonGroup.MaxTriangles)
-					{
-						RenderingSection.NumTriangles = RenderingPolygonGroup.MaxTriangles;
-					}
-
 					if( !EditableMesh->IsPreviewingSubdivisions() )
 					{
 						const uint32 FirstIndex = RenderingSection.FirstIndex;
@@ -1433,6 +1427,7 @@ void UEditableStaticMeshAdapter::OnRetriangulatePolygons( const UEditableMesh* E
 					}
 				}
 
+				// Add triangles to rendering section
 				if( !EditableMesh->IsPreviewingSubdivisions() )
 				{
 					for( int32 TriangleToAddNumber = 0; TriangleToAddNumber < NumNewTriangles; ++TriangleToAddNumber )
@@ -1454,6 +1449,7 @@ void UEditableStaticMeshAdapter::OnRetriangulatePolygons( const UEditableMesh* E
 
 					RenderingSection.MinVertexIndex = MinVertexIndex;
 					RenderingSection.MaxVertexIndex = MaxVertexIndex;
+					check(RenderingSection.NumTriangles <=  (uint32) RenderingPolygonGroup.MaxTriangles);
 				}
 			}
 		}
