First pass of MeshDescription API and format refactor.

diff --git a/Engine/Plugins/Experimental/GeometryProcessing/Source/MeshConversion/Private/DynamicMeshToMeshDescription.cpp b/Engine/Plugins/Experimental/GeometryProcessing/Source/MeshConversion/Private/DynamicMeshToMeshDescription.cpp
--- a/Engine/Plugins/Experimental/GeometryProcessing/Source/MeshConversion/Private/DynamicMeshToMeshDescription.cpp
+++ b/Engine/Plugins/Experimental/GeometryProcessing/Source/MeshConversion/Private/DynamicMeshToMeshDescription.cpp
@@ -25,7 +25,7 @@ namespace DynamicMeshToMeshDescriptionConversionHelper
 		{
 			FPolygonGroupID PolygonGroupID = MeshOutArg.GetPolygonPolygonGroup(PolygonID);
 
-			const TArray<FTriangleID>& TriangleIDs = MeshOutArg.GetPolygonTriangleIDs(PolygonID);
+			TArrayView<const FTriangleID> TriangleIDs = MeshOutArg.GetPolygonTriangleIDs(PolygonID);
 			int NumTriangles = TriangleIDs.Num();
 			for (int TriIdx = 0; TriIdx < NumTriangles; ++TriIdx, ++MeshInTriIdx)
 			{
@@ -261,7 +261,7 @@ void FDynamicMeshToMeshDescription::Convert_SharedInstances(const FDynamicMesh3*
 	{
 		FVertexInstanceID V[3];
 
-		Tri() : V{ FVertexInstanceID::Invalid,FVertexInstanceID::Invalid,FVertexInstanceID::Invalid }
+		Tri() : V{ INDEX_NONE,INDEX_NONE,INDEX_NONE }
 		{}
 	};
 	TArray<Tri> TriVertInstances;
@@ -278,7 +278,7 @@ void FDynamicMeshToMeshDescription::Convert_SharedInstances(const FDynamicMesh3*
 			FIndex3i Tri = MeshIn->GetTriangle(TriID);
 			int SubIdx = IndexUtil::FindTriIndex(VertID, Tri);
 
-			int32 FoundInstance = FVertexInstanceID::Invalid.GetValue();
+			int32 FoundInstance = INDEX_NONE;
 			for (int KIItemIdx = 0; KIItemIdx < KnownInstanceIDs.Num(); KIItemIdx += KIItemLen)
 			{
 				int KIItemInternalIdx = KIItemIdx;
@@ -307,7 +307,7 @@ void FDynamicMeshToMeshDescription::Convert_SharedInstances(const FDynamicMesh3*
 				check(KIItemInternalIdx == KIItemIdx + KIItemLen);
 				break;
 			}
-			if (FoundInstance == FVertexInstanceID::Invalid.GetValue())
+			if (FoundInstance == INDEX_NONE)
 			{
 				FVertexInstanceID NewInstanceID = Builder.AppendInstance(MapV[VertID]);
 				if (NormalOverlay)
diff --git a/Engine/Plugins/Experimental/GeometryProcessing/Source/MeshConversion/Private/MeshDescriptionBuilder.cpp b/Engine/Plugins/Experimental/GeometryProcessing/Source/MeshConversion/Private/MeshDescriptionBuilder.cpp
--- a/Engine/Plugins/Experimental/GeometryProcessing/Source/MeshConversion/Private/MeshDescriptionBuilder.cpp
+++ b/Engine/Plugins/Experimental/GeometryProcessing/Source/MeshConversion/Private/MeshDescriptionBuilder.cpp
@@ -118,7 +118,7 @@ void FMeshDescriptionBuilder::SetInstanceNormal(const FVertexInstanceID& Instanc
 
 void FMeshDescriptionBuilder::SetInstanceUV(const FVertexInstanceID& InstanceID, const FVector2D& InstanceUV, int32 UVLayerIndex)
 {
-	if (InstanceUVs.IsValid() && ensure(UVLayerIndex < InstanceUVs.GetNumIndices()))
+	if (InstanceUVs.IsValid() && ensure(UVLayerIndex < InstanceUVs.GetNumChannels()))
 	{
 		InstanceUVs.Set(InstanceID, UVLayerIndex, InstanceUV); 
 	}
@@ -129,7 +129,7 @@ void FMeshDescriptionBuilder::SetNumUVLayers(int32 NumUVLayers)
 {
 	if (ensure(InstanceUVs.IsValid()))
 	{
-		InstanceUVs.SetNumIndices(NumUVLayers);
+		InstanceUVs.SetNumChannels(NumUVLayers);
 	}
 }
 
diff --git a/Engine/Plugins/Experimental/GeometryProcessing/Source/MeshConversion/Private/MeshDescriptionToDynamicMesh.cpp b/Engine/Plugins/Experimental/GeometryProcessing/Source/MeshConversion/Private/MeshDescriptionToDynamicMesh.cpp
--- a/Engine/Plugins/Experimental/GeometryProcessing/Source/MeshConversion/Private/MeshDescriptionToDynamicMesh.cpp
+++ b/Engine/Plugins/Experimental/GeometryProcessing/Source/MeshConversion/Private/MeshDescriptionToDynamicMesh.cpp
@@ -181,7 +181,7 @@ void FMeshDescriptionToDynamicMesh::Convert(const FMeshDescription* MeshIn, FDyn
 	{
 		int32 PolygonGroupID = MeshIn->GetPolygonPolygonGroup(PolygonID).GetValue();
 
-		const TArray<FTriangleID>& TriangleIDs = MeshIn->GetPolygonTriangleIDs(PolygonID);
+		TArrayView<const FTriangleID> TriangleIDs = MeshIn->GetPolygonTriangleIDs(PolygonID);
 		int NumTriangles = TriangleIDs.Num();
 		for (int TriIdx = 0; TriIdx < NumTriangles; ++TriIdx)
 		{
@@ -284,7 +284,7 @@ void FMeshDescriptionToDynamicMesh::Convert(const FMeshDescription* MeshIn, FDyn
 	// Enable relevant attributes and initialize UV/Normal welders
 	// 
 
-	int NumUVLayers = InstanceUVs.GetNumIndices();
+	int NumUVLayers = InstanceUVs.GetNumChannels();
 	TArray<FDynamicMeshUVOverlay*> UVOverlays;
 	TArray<FUVWelder> UVWelders;
 	FDynamicMeshNormalOverlay* NormalOverlay = nullptr;
@@ -431,7 +431,7 @@ void FMeshDescriptionToDynamicMesh::CopyTangents(const FMeshDescription* SourceM
 	for (int TriID : TargetMesh->TriangleIndicesItr())
 	{
 		FIndex2i PolyTriIdx = TriToPolyTriMap[TriID];
-		const TArray<FTriangleID>& TriangleIDs = SourceMesh->GetPolygonTriangleIDs( FPolygonID(PolyTriIdx.A) );
+		TArrayView<const FTriangleID> TriangleIDs = SourceMesh->GetPolygonTriangleIDs( FPolygonID(PolyTriIdx.A) );
 		const FTriangleID TriangleID = TriangleIDs[PolyTriIdx.B];
 		TArrayView<const FVertexInstanceID> InstanceTri = SourceMesh->GetTriangleVertexInstances(TriangleID);
 		for (int j = 0; j < 3; ++j)
diff --git a/Engine/Plugins/Experimental/GeometryProcessing/Source/MeshConversion/Public/MeshDescriptionAdapter.h b/Engine/Plugins/Experimental/GeometryProcessing/Source/MeshConversion/Public/MeshDescriptionAdapter.h
--- a/Engine/Plugins/Experimental/GeometryProcessing/Source/MeshConversion/Public/MeshDescriptionAdapter.h
+++ b/Engine/Plugins/Experimental/GeometryProcessing/Source/MeshConversion/Public/MeshDescriptionAdapter.h
@@ -64,7 +64,7 @@ public:
 	}
 	FIndex3i GetTriangle(int32 IDValue) const
 	{
-		const TStaticArray<FVertexID, 3> TriVertIDs = Mesh->GetTriangleVertices(FTriangleID(IDValue));
+		TArrayView<const FVertexID> TriVertIDs = Mesh->GetTriangleVertices(FTriangleID(IDValue));
 		return FIndex3i(TriVertIDs[0].GetValue(), TriVertIDs[1].GetValue(), TriVertIDs[2].GetValue());
 	}
 	FVector3d GetVertex(int32 IDValue) const
@@ -74,7 +74,7 @@ public:
 
 	inline void GetTriVertices(int32 IDValue, FVector3d& V0, FVector3d& V1, FVector3d& V2) const
 	{
-		const TStaticArray<FVertexID, 3> TriVertIDs = Mesh->GetTriangleVertices(FTriangleID(IDValue));
+		TArrayView<const FVertexID> TriVertIDs = Mesh->GetTriangleVertices(FTriangleID(IDValue));
 		V0 = FVector3d(VertexPositions[TriVertIDs[0]]);
 		V1 = FVector3d(VertexPositions[TriVertIDs[1]]);
 		V2 = FVector3d(VertexPositions[TriVertIDs[2]]);
@@ -155,7 +155,7 @@ public:
 	}
 	FIndex3i GetTriangle(int32 IDValue) const
 	{
-		const TStaticArray<FVertexID, 3> TriVertIDs = Mesh->GetTriangleVertices(FTriangleID(IDValue));
+		TArrayView<const FVertexID> TriVertIDs = Mesh->GetTriangleVertices(FTriangleID(IDValue));
 		return FIndex3i(TriVertIDs[0].GetValue(), TriVertIDs[1].GetValue(), TriVertIDs[2].GetValue());
 	}
 	FVector3d GetVertex(int32 IDValue) const
@@ -169,7 +169,7 @@ public:
 
 	inline void GetTriVertices(int32 IDValue, FVector3d& V0, FVector3d& V1, FVector3d& V2) const
 	{
-		const TStaticArray<FVertexID, 3> TriVertIDs = Mesh->GetTriangleVertices(FTriangleID(IDValue));
+		TArrayView<const FVertexID> TriVertIDs = Mesh->GetTriangleVertices(FTriangleID(IDValue));
 		V0 = FVector3d(VertexPositions[TriVertIDs[0]]);
 		V1 = FVector3d(VertexPositions[TriVertIDs[1]]);
 		V2 = FVector3d(VertexPositions[TriVertIDs[2]]);
diff --git a/Engine/Plugins/Experimental/MeshModelingToolset/Source/MeshModelingToolsEditorOnly/Private/AttributeEditorTool.cpp b/Engine/Plugins/Experimental/MeshModelingToolset/Source/MeshModelingToolsEditorOnly/Private/AttributeEditorTool.cpp
--- a/Engine/Plugins/Experimental/MeshModelingToolset/Source/MeshModelingToolsEditorOnly/Private/AttributeEditorTool.cpp
+++ b/Engine/Plugins/Experimental/MeshModelingToolset/Source/MeshModelingToolsEditorOnly/Private/AttributeEditorTool.cpp
@@ -98,7 +98,7 @@ void UAttributeEditorTool::Setup()
 		FMeshDescription* MeshDescription = ComponentTarget->GetMesh();
 		TVertexInstanceAttributesRef<FVector2D> InstanceUVs =
 			MeshDescription->VertexInstanceAttributes().GetAttributesRef<FVector2D>(MeshAttribute::VertexInstance::TextureCoordinate);
-		NumUVLayers = FMath::Max(NumUVLayers, InstanceUVs.GetNumIndices());
+		NumUVLayers = FMath::Max(NumUVLayers, InstanceUVs.GetNumChannels());
 	}
 
 	NormalsActions = NewObject<UAttributeEditorNormalsActions>(this);
@@ -269,22 +269,13 @@ static bool RemoveAttribute(FMeshDescription* Mesh, EAttributeEditorElementType
 static bool IsReservedName(FName AttributeName)
 {
 	return (AttributeName == MeshAttribute::Vertex::Position)
-		|| (AttributeName == MeshAttribute::Vertex::CornerSharpness)
 		|| (AttributeName == MeshAttribute::VertexInstance::TextureCoordinate)
 		|| (AttributeName == MeshAttribute::VertexInstance::Normal)
 		|| (AttributeName == MeshAttribute::VertexInstance::Tangent)
 		|| (AttributeName == MeshAttribute::VertexInstance::BinormalSign)
 		|| (AttributeName == MeshAttribute::VertexInstance::Color)
 		|| (AttributeName == MeshAttribute::Edge::IsHard)
-		|| (AttributeName == MeshAttribute::Edge::IsUVSeam)
-		|| (AttributeName == MeshAttribute::Edge::CreaseSharpness)
-		|| (AttributeName == MeshAttribute::Polygon::Normal)
-		|| (AttributeName == MeshAttribute::Polygon::Tangent)
-		|| (AttributeName == MeshAttribute::Polygon::Binormal)
-		|| (AttributeName == MeshAttribute::Polygon::Center)
-		|| (AttributeName == MeshAttribute::PolygonGroup::ImportedMaterialSlotName)
-		|| (AttributeName == MeshAttribute::PolygonGroup::EnableCollision)
-		|| (AttributeName == MeshAttribute::PolygonGroup::CastShadow);
+		|| (AttributeName == MeshAttribute::PolygonGroup::ImportedMaterialSlotName);
 }
 
 
@@ -394,7 +385,7 @@ void UAttributeEditorTool::ClearNormals()
 					EdgeHardnesses[ElID] = false;
 				}
 			}
-			FStaticMeshOperations::ComputePolygonTangentsAndNormals(*CommitParams.MeshDescription, FMathf::Epsilon);
+			FStaticMeshOperations::ComputeTriangleTangentsAndNormals(*CommitParams.MeshDescription, FMathf::Epsilon);
 			FStaticMeshOperations::RecomputeNormalsAndTangentsIfNeeded(*CommitParams.MeshDescription, EComputeNTBsFlags::WeightedNTBs | EComputeNTBsFlags::Normals);
 		});
 	}
@@ -429,7 +420,7 @@ void UAttributeEditorTool::ClearUVs(bool bSelectedOnly)
 			for (int LayerIndex = 7; LayerIndex >= 0; LayerIndex--)
 			{
 				bool bRemove = (bSelectedOnly == false) || RemoveLayers[LayerIndex];
-				if (bRemove && LayerIndex < InstanceUVs.GetNumIndices())
+				if (bRemove && LayerIndex < InstanceUVs.GetNumChannels())
 				{
 					if (!FStaticMeshOperations::RemoveUVChannel(*CommitParams.MeshDescription, LayerIndex))
 					{
@@ -513,6 +504,7 @@ void UAttributeEditorTool::DeleteAttribute()
 	FMeshDescription* CurMesh = ComponentTargets[0]->GetMesh();
 	FName SelectedName(ModifyAttributeProps->Attribute);
 
+	// @todo: introduce 'reserved' attribute flag which we can check here
 	if (IsReservedName(SelectedName))
 	{
 		GetToolManager()->DisplayMessage(LOCTEXT("CannotDeleteReservedNameError", "Cannot delete required mesh Attributes"), EToolMessageLevel::UserWarning);
diff --git a/Engine/Plugins/Experimental/ProxyLODPlugin/Source/ProxyLOD/Private/ProxyLODMeshConvertUtils.cpp b/Engine/Plugins/Experimental/ProxyLODPlugin/Source/ProxyLOD/Private/ProxyLODMeshConvertUtils.cpp
--- a/Engine/Plugins/Experimental/ProxyLODPlugin/Source/ProxyLOD/Private/ProxyLODMeshConvertUtils.cpp
+++ b/Engine/Plugins/Experimental/ProxyLODPlugin/Source/ProxyLOD/Private/ProxyLODMeshConvertUtils.cpp
@@ -22,7 +22,6 @@ void ProxyLOD::MixedPolyMeshToRawMesh(const FMixedPolyMesh& SimpleMesh, FMeshDes
 
 	TVertexAttributesRef<FVector> VertexPositions = DstRawMesh.VertexAttributes().GetAttributesRef<FVector>(MeshAttribute::Vertex::Position);
 	TEdgeAttributesRef<bool> EdgeHardnesses = DstRawMesh.EdgeAttributes().GetAttributesRef<bool>(MeshAttribute::Edge::IsHard);
-	TEdgeAttributesRef<float> EdgeCreaseSharpnesses = DstRawMesh.EdgeAttributes().GetAttributesRef<float>(MeshAttribute::Edge::CreaseSharpness);
 	TPolygonGroupAttributesRef<FName> PolygonGroupImportedMaterialSlotNames = DstRawMesh.PolygonGroupAttributes().GetAttributesRef<FName>(MeshAttribute::PolygonGroup::ImportedMaterialSlotName);
 	TVertexInstanceAttributesRef<FVector> VertexInstanceNormals = DstRawMesh.VertexInstanceAttributes().GetAttributesRef<FVector>(MeshAttribute::VertexInstance::Normal);
 	TVertexInstanceAttributesRef<FVector> VertexInstanceTangents = DstRawMesh.VertexInstanceAttributes().GetAttributesRef<FVector>(MeshAttribute::VertexInstance::Tangent);
@@ -41,12 +40,12 @@ void ProxyLOD::MixedPolyMeshToRawMesh(const FMixedPolyMesh& SimpleMesh, FMeshDes
 	// Each Triangle has 3 corners
 	const uint32 DstNumIndexes = 3 * DstNumTris;
 
-	if (VertexInstanceUVs.GetNumIndices() < 1)
+	if (VertexInstanceUVs.GetNumChannels() < 1)
 	{
-		VertexInstanceUVs.SetNumIndices(1);
+		VertexInstanceUVs.SetNumChannels(1);
 	}
 
-	FPolygonGroupID PolygonGroupID = FPolygonGroupID::Invalid;
+	FPolygonGroupID PolygonGroupID = INDEX_NONE;
 	if (DstRawMesh.PolygonGroups().Num() == 0)
 	{
 		PolygonGroupID = DstRawMesh.CreatePolygonGroup();
@@ -71,7 +70,7 @@ void ProxyLOD::MixedPolyMeshToRawMesh(const FMixedPolyMesh& SimpleMesh, FMeshDes
 	}
 
 	// Connectivity:
-	auto CreateTriangle = [&DstRawMesh, PolygonGroupID, &VertexInstanceNormals, &VertexInstanceTangents, &VertexInstanceBinormalSigns, &VertexInstanceColors, &VertexInstanceUVs, &EdgeHardnesses, &EdgeCreaseSharpnesses](FVertexID TriangleIndex[3])
+	auto CreateTriangle = [&DstRawMesh, PolygonGroupID, &VertexInstanceNormals, &VertexInstanceTangents, &VertexInstanceBinormalSigns, &VertexInstanceColors, &VertexInstanceUVs, &EdgeHardnesses](FVertexID TriangleIndex[3])
 	{
 		TArray<FVertexInstanceID> VertexInstanceIDs;
 		VertexInstanceIDs.SetNum(3);
@@ -137,7 +136,6 @@ void ProxyLOD::AOSMeshToRawMesh(const FAOSMesh& AOSMesh, FMeshDescription& OutRa
 
 	TVertexAttributesRef<FVector> VertexPositions = OutRawMesh.VertexAttributes().GetAttributesRef<FVector>(MeshAttribute::Vertex::Position);
 	TEdgeAttributesRef<bool> EdgeHardnesses = OutRawMesh.EdgeAttributes().GetAttributesRef<bool>(MeshAttribute::Edge::IsHard);
-	TEdgeAttributesRef<float> EdgeCreaseSharpnesses = OutRawMesh.EdgeAttributes().GetAttributesRef<float>(MeshAttribute::Edge::CreaseSharpness);
 	TPolygonGroupAttributesRef<FName> PolygonGroupImportedMaterialSlotNames = OutRawMesh.PolygonGroupAttributes().GetAttributesRef<FName>(MeshAttribute::PolygonGroup::ImportedMaterialSlotName);
 	TVertexInstanceAttributesRef<FVector> VertexInstanceNormals = OutRawMesh.VertexInstanceAttributes().GetAttributesRef<FVector>(MeshAttribute::VertexInstance::Normal);
 	TVertexInstanceAttributesRef<FVector> VertexInstanceTangents = OutRawMesh.VertexInstanceAttributes().GetAttributesRef<FVector>(MeshAttribute::VertexInstance::Tangent);
@@ -148,12 +146,12 @@ void ProxyLOD::AOSMeshToRawMesh(const FAOSMesh& AOSMesh, FMeshDescription& OutRa
 	const uint32 DstNumPositions = AOSMesh.GetNumVertexes();
 	const uint32 DstNumIndexes = AOSMesh.GetNumIndexes();
 
-	if (VertexInstanceUVs.GetNumIndices() < 1)
+	if (VertexInstanceUVs.GetNumChannels() < 1)
 	{
-		VertexInstanceUVs.SetNumIndices(1);
+		VertexInstanceUVs.SetNumChannels(1);
 	}
 
-	FPolygonGroupID PolygonGroupID = FPolygonGroupID::Invalid;
+	FPolygonGroupID PolygonGroupID = INDEX_NONE;
 	if (OutRawMesh.PolygonGroups().Num() == 0)
 	{
 		PolygonGroupID = OutRawMesh.CreatePolygonGroup();
@@ -184,7 +182,7 @@ void ProxyLOD::AOSMeshToRawMesh(const FAOSMesh& AOSMesh, FMeshDescription& OutRa
 	const uint32* AOSIndexes = AOSMesh.Indexes;
 
 	// Connectivity: 
-	auto CreateTriangle = [&OutRawMesh, PolygonGroupID, &VertexInstanceNormals, &VertexInstanceTangents, &VertexInstanceBinormalSigns, &VertexInstanceColors, &VertexInstanceUVs, &EdgeHardnesses, &EdgeCreaseSharpnesses](const FVertexID TriangleIndex[3], const FVector Normals[3])
+	auto CreateTriangle = [&OutRawMesh, PolygonGroupID, &VertexInstanceNormals, &VertexInstanceTangents, &VertexInstanceBinormalSigns, &VertexInstanceColors, &VertexInstanceUVs, &EdgeHardnesses](const FVertexID TriangleIndex[3], const FVector Normals[3])
 	{
 		TArray<FVertexInstanceID> VertexInstanceIDs;
 		VertexInstanceIDs.SetNum(3);
@@ -228,7 +226,6 @@ void ProxyLOD::VertexDataMeshToRawMesh(const FVertexDataMesh& SrcVertexDataMesh,
 
 	TVertexAttributesRef<FVector> VertexPositions = OutRawMesh.VertexAttributes().GetAttributesRef<FVector>(MeshAttribute::Vertex::Position);
 	TEdgeAttributesRef<bool> EdgeHardnesses = OutRawMesh.EdgeAttributes().GetAttributesRef<bool>(MeshAttribute::Edge::IsHard);
-	TEdgeAttributesRef<float> EdgeCreaseSharpnesses = OutRawMesh.EdgeAttributes().GetAttributesRef<float>(MeshAttribute::Edge::CreaseSharpness);
 	TPolygonGroupAttributesRef<FName> PolygonGroupImportedMaterialSlotNames = OutRawMesh.PolygonGroupAttributes().GetAttributesRef<FName>(MeshAttribute::PolygonGroup::ImportedMaterialSlotName);
 	TVertexInstanceAttributesRef<FVector> VertexInstanceNormals = OutRawMesh.VertexInstanceAttributes().GetAttributesRef<FVector>(MeshAttribute::VertexInstance::Normal);
 	TVertexInstanceAttributesRef<FVector> VertexInstanceTangents = OutRawMesh.VertexInstanceAttributes().GetAttributesRef<FVector>(MeshAttribute::VertexInstance::Tangent);
@@ -240,7 +237,7 @@ void ProxyLOD::VertexDataMeshToRawMesh(const FVertexDataMesh& SrcVertexDataMesh,
 	const uint32 DstNumIndexes = SrcVertexDataMesh.Indices.Num();
 	const uint32 SrcNumTriangles = DstNumIndexes / 3;
 	
-	FPolygonGroupID PolygonGroupID = FPolygonGroupID::Invalid;
+	FPolygonGroupID PolygonGroupID = INDEX_NONE;
 	if (OutRawMesh.PolygonGroups().Num() == 0)
 	{
 		PolygonGroupID = OutRawMesh.CreatePolygonGroup();
@@ -272,7 +269,7 @@ void ProxyLOD::VertexDataMeshToRawMesh(const FVertexDataMesh& SrcVertexDataMesh,
 	const bool bSrcHasTangentSpace = SrcVertexDataMesh.Tangent.Num() != 0 && SrcVertexDataMesh.BiTangent.Num() != 0 && SrcVertexDataMesh.Normal.Num() != 0;
 
 	// Connectivity:
-	auto CreateTriangle = [&OutRawMesh, &SrcVertexDataMesh, bSrcHasTangentSpace, &VertexIDMap, PolygonGroupID, &VertexInstanceNormals, &VertexInstanceTangents, &VertexInstanceBinormalSigns, &VertexInstanceColors, &VertexInstanceUVs, &EdgeHardnesses, &EdgeCreaseSharpnesses](const uint32 TriangleIndices[3])
+	auto CreateTriangle = [&OutRawMesh, &SrcVertexDataMesh, bSrcHasTangentSpace, &VertexIDMap, PolygonGroupID, &VertexInstanceNormals, &VertexInstanceTangents, &VertexInstanceBinormalSigns, &VertexInstanceColors, &VertexInstanceUVs, &EdgeHardnesses](const uint32 TriangleIndices[3])
 	{
 		int32 TriangleIndex = TriangleIndices[0] / 3;
 		
@@ -430,7 +427,7 @@ void ProxyLOD::RawMeshToVertexDataMesh(const FMeshDescription& SrcRawMesh, FVert
 				// Copy the UVs:
 				// NB: The UVs is stored per-index in the raw mesh, but only per-vertex in the vertex data mesh.
 				// We assume that the raw mesh per-index data is really duplicated per-vertex data!
-				if (VertexInstanceUVs.GetNumIndices() == 0)
+				if (VertexInstanceUVs.GetNumChannels() == 0)
 				{
 					DstUVs[DstIndices[VertexInstanceCount]] = FVector2D(0.0f, 0.0f);
 				}
diff --git a/Engine/Plugins/Experimental/ProxyLODPlugin/Source/ProxyLOD/Private/ProxyLODMeshTypes.cpp b/Engine/Plugins/Experimental/ProxyLODPlugin/Source/ProxyLOD/Private/ProxyLODMeshTypes.cpp
--- a/Engine/Plugins/Experimental/ProxyLODPlugin/Source/ProxyLOD/Private/ProxyLODMeshTypes.cpp
+++ b/Engine/Plugins/Experimental/ProxyLODPlugin/Source/ProxyLOD/Private/ProxyLODMeshTypes.cpp
@@ -26,7 +26,7 @@ void FMeshDescriptionAdapter::InitializeCacheData()
 
 	for (const FPolygonID PolygonID : RawMesh->Polygons().GetElementIDs())
 	{
-		const TArray<FTriangleID>& TriangleIDs = RawMesh->GetPolygonTriangleIDs(PolygonID);
+		TArrayView<const FTriangleID> TriangleIDs = RawMesh->GetPolygonTriangleIDs(PolygonID);
 		for (const FTriangleID TriangleID : TriangleIDs)
 		{
 			IndexBuffer.Add(RawMesh->GetTriangleVertexInstance(TriangleID, 0));
@@ -81,7 +81,7 @@ void FMeshDescriptionArrayAdapter::Construct(int32 MeshCount, TFunctionRef<const
 		IndexBuffer.reserve(MeshPolyCount * 3);
 		for (const FPolygonID PolygonID : RawMesh->Polygons().GetElementIDs())
 		{
-			const TArray<FTriangleID>& TriangleIDs = RawMesh->GetPolygonTriangleIDs(PolygonID);
+			TArrayView<const FTriangleID> TriangleIDs = RawMesh->GetPolygonTriangleIDs(PolygonID);
 			for (const FTriangleID TriangleID : TriangleIDs)
 			{
 				IndexBuffer.push_back(RawMesh->GetTriangleVertexInstance(TriangleID, 0));
@@ -284,7 +284,7 @@ FMeshDescriptionArrayAdapter::FRawPoly FMeshDescriptionArrayAdapter::GetRawPoly(
 				// Copy Texture coords
 				for (int Idx = 0; Idx < MAX_MESH_TEXTURE_COORDS_MD; ++Idx)
 				{
-					if (AttributesGetter->VertexInstanceUVs.GetNumIndices() > Idx)
+					if (AttributesGetter->VertexInstanceUVs.GetNumChannels() > Idx)
 					{
 						RawPoly.WedgeTexCoords[Idx][i] = AttributesGetter->VertexInstanceUVs.Get(VertexInstanceIDs[i], Idx);
 					}
diff --git a/Engine/Plugins/Importers/USDImporter/Source/USDUtilities/Private/USDGeomMeshConversion.cpp b/Engine/Plugins/Importers/USDImporter/Source/USDUtilities/Private/USDGeomMeshConversion.cpp
--- a/Engine/Plugins/Importers/USDImporter/Source/USDUtilities/Private/USDGeomMeshConversion.cpp
+++ b/Engine/Plugins/Importers/USDImporter/Source/USDUtilities/Private/USDGeomMeshConversion.cpp
@@ -204,14 +204,14 @@ bool UsdToUnreal::ConvertGeomMesh( const pxr::UsdGeomMesh& UsdMesh, FMeshDescrip
 
 		// When importing multiple mesh pieces to the same static mesh.  Ensure each mesh piece has the same number of Uv's
 		{
-			int32 ExistingUVCount = MeshDescriptionUVs.GetNumIndices();
+			int32 ExistingUVCount = MeshDescriptionUVs.GetNumChannels();
 			int32 NumUVs = FMath::Max(UVSets.Num(), ExistingUVCount);
 			NumUVs = FMath::Min<int32>(MAX_MESH_TEXTURE_COORDS_MD, NumUVs);
 			// At least one UV set must exist.  
 			NumUVs = FMath::Max<int32>(1, NumUVs);
 
 			//Make sure all Vertex instance have the correct number of UVs
-			MeshDescriptionUVs.SetNumIndices(NumUVs);
+			MeshDescriptionUVs.SetNumChannels(NumUVs);
 		}
 
 		TVertexInstanceAttributesRef< FVector > MeshDescriptionNormals = StaticMeshAttributes.GetVertexInstanceNormals();
@@ -375,11 +375,14 @@ bool UsdToUnreal::ConvertGeomMesh( const pxr::UsdGeomMesh& UsdMesh, FMeshDescrip
 
 			FPolygonGroupID PolygonGroupID = PolygonGroupMapping[ RealMaterialIndex ];
 			// Insert a polygon into the mesh
-			const FPolygonID NewPolygonID = MeshDescription.CreatePolygon( PolygonGroupID, CornerInstanceIDs );
 			if ( bFlipThisGeometry )
 			{
-				MeshDescription.ReversePolygonFacing( NewPolygonID );
+				for (int32 i = 0; i < CornerInstanceIDs.Num() / 2; ++i)
+				{
+					Swap(CornerInstanceIDs[i], CornerInstanceIDs[CornerInstanceIDs.Num() - i - 1]);
+				}
 			}
+			const FPolygonID NewPolygonID = MeshDescription.CreatePolygon( PolygonGroupID, CornerInstanceIDs );
 		}
 	}
 
diff --git a/Engine/Plugins/Runtime/EditableMesh/EditableMesh.uplugin b/Engine/Plugins/Runtime/EditableMesh/EditableMesh.uplugin
--- a/Engine/Plugins/Runtime/EditableMesh/EditableMesh.uplugin
+++ b/Engine/Plugins/Runtime/EditableMesh/EditableMesh.uplugin
@@ -10,15 +10,16 @@
 	"DocsURL" : "",
 	"MarketplaceURL" : "",
 	"SupportURL" : "",
-	"EnabledByDefault" : true,
+	"EnabledByDefault" : false,
 	"CanContainContent" : false,
 	"IsBetaVersion" : false,
-	"Installed" : false,
-  "Modules": [
+	"Installed": false,
+	"Hidden":  true,
+	"Modules": [
     {
-      "Name": "EditableMesh",
-      "Type": "Runtime",
-      "LoadingPhase": "Default"
+		"Name": "EditableMesh",
+		"Type": "Runtime",
+		"LoadingPhase": "Default"
     }
   ]
 }
\ No newline at end of file
diff --git a/Engine/Plugins/Runtime/EditableMesh/Source/EditableMesh/EditableGeometryCollectionAdapter.cpp b/Engine/Plugins/Runtime/EditableMesh/Source/EditableMesh/EditableGeometryCollectionAdapter.cpp
--- a/Engine/Plugins/Runtime/EditableMesh/Source/EditableMesh/EditableGeometryCollectionAdapter.cpp
+++ b/Engine/Plugins/Runtime/EditableMesh/Source/EditableMesh/EditableGeometryCollectionAdapter.cpp
@@ -85,7 +85,7 @@ void UEditableGeometryCollectionAdapter::InitEditableGeometryCollection( UEditab
 
 				// Store off the number of texture coordinates in this mesh
 				const int32 NumUVs = 1;
-				VertexInstanceUVs.SetNumIndices(NumUVs);
+				VertexInstanceUVs.SetNumChannels(NumUVs);
 				EditableMesh->TextureCoordinateCount = NumUVs;
 
 				// Vertex Positions
@@ -203,7 +203,7 @@ void UEditableGeometryCollectionAdapter::InitEditableGeometryCollection( UEditab
 
 									FEdgeID NewEdgeID = MeshDescription->GetVertexPairEdge(VertexID0, VertexID1);
 
-									if (NewEdgeID == FEdgeID::Invalid)
+									if (NewEdgeID == INDEX_NONE)
 									{
 										NewEdgeID = MeshDescription->CreateEdge(VertexID0, VertexID1);
 									}
@@ -217,7 +217,7 @@ void UEditableGeometryCollectionAdapter::InitEditableGeometryCollection( UEditab
 							const FAdaptorTriangleID NewTriangleID = FAdaptorTriangleID(TriangleGroupIndex++);
 
 							NewRenderingPolygonGroup.Triangles.Insert(NewTriangleID);
-							FMeshTriangle& NewTriangle = NewRenderingPolygonGroup.Triangles[NewTriangleID];
+							FAdaptorPolygon2Group::FMeshTriangle& NewTriangle = NewRenderingPolygonGroup.Triangles[NewTriangleID];
 							for (int32 TriangleVertexIndex = 0; TriangleVertexIndex < 3; ++TriangleVertexIndex)
 							{
 								NewTriangle.SetVertexInstanceID(TriangleVertexIndex, TriangleVertexInstanceIDs[TriangleVertexIndex]);
@@ -324,10 +324,10 @@ void UEditableGeometryCollectionAdapter::InitializeFromEditableMesh( const UEdit
 		FAdaptorPolygon& RenderingPolygon = RenderingPolygons[ PolygonID ];
 		RenderingPolygon.PolygonGroupID = PolygonGroupID;
 
-		const TArray<FTriangleID>& TriangleIDs = MeshDescription->GetPolygonTriangleIDs( PolygonID );
+		TArrayView<const FTriangleID> TriangleIDs = MeshDescription->GetPolygonTriangleIDs( PolygonID );
 		for( const FTriangleID TriangleID : TriangleIDs )
 		{
-			FMeshTriangle Triangle;
+			FAdaptorPolygon2Group::FMeshTriangle Triangle;
 			Triangle.SetVertexInstanceID( 0, MeshDescription->GetTriangleVertexInstance( TriangleID, 0 ) );
 			Triangle.SetVertexInstanceID( 1, MeshDescription->GetTriangleVertexInstance( TriangleID, 1 ) );
 			Triangle.SetVertexInstanceID( 2, MeshDescription->GetTriangleVertexInstance( TriangleID, 2 ) );
@@ -406,7 +406,7 @@ void UEditableGeometryCollectionAdapter::OnRebuildRenderMesh(const UEditableMesh
 				const FPolygonGroupID& PolygonGroupID = MeshDescription->GetPolygonPolygonGroup(PolygonID);
 				int32 PolygonGroupIDValue = PolygonGroupID.GetValue();
 				int32 PolygonIDValue = PolygonID.GetValue();
-				const TArray<FTriangleID>& TriangleIDs = MeshDescription->GetPolygonTriangleIDs(PolygonID);
+				TArrayView<const FTriangleID> TriangleIDs = MeshDescription->GetPolygonTriangleIDs(PolygonID);
 				for (const FTriangleID TriangleID : TriangleIDs)
 				{
 					uint32 ElementIndex = Collection->AddElements(1, FGeometryCollection::FacesGroup);
@@ -874,7 +874,7 @@ void UEditableGeometryCollectionAdapter::DeletePolygonTriangles( const UEditable
 				// vertex instances will need to be removed from their corresponding vertex
 				for (const FAdaptorTriangleID TriangleIndexToRemove : Polygon.TriangulatedPolygonTriangleIndices)
 				{
-					const FMeshTriangle& TriToRemove = RenderingPolygonGroup.Triangles[TriangleIndexToRemove];
+					const FAdaptorPolygon2Group::FMeshTriangle& TriToRemove = RenderingPolygonGroup.Triangles[TriangleIndexToRemove];
 
 					// sanity check
 					FIntVector TriIndices = GCIndices[TriangleIndexToRemove.GetValue()];
@@ -941,7 +941,7 @@ FPolygonGroupID UEditableGeometryCollectionAdapter::GetSectionForRenderingSectio
 		}
 	}
 
-	return FPolygonGroupID::Invalid;
+	return INDEX_NONE;
 }
 
 
diff --git a/Engine/Plugins/Runtime/EditableMesh/Source/EditableMesh/EditableGeometryCollectionAdapter.h b/Engine/Plugins/Runtime/EditableMesh/Source/EditableMesh/EditableGeometryCollectionAdapter.h
--- a/Engine/Plugins/Runtime/EditableMesh/Source/EditableMesh/EditableGeometryCollectionAdapter.h
+++ b/Engine/Plugins/Runtime/EditableMesh/Source/EditableMesh/EditableGeometryCollectionAdapter.h
@@ -71,6 +71,58 @@ struct FAdaptorPolygon2Group
 {
 	GENERATED_BODY()
 
+	struct FMeshTriangle
+	{
+		friend struct FMeshDescription;
+
+		FMeshTriangle()
+		{
+			VertexInstanceIDs[0] = INDEX_NONE;
+			VertexInstanceIDs[1] = INDEX_NONE;
+			VertexInstanceIDs[2] = INDEX_NONE;
+		}
+
+	private:
+
+		/** Vertex instance IDs that make up this triangle.  Indices must be ordered counter-clockwise. */
+		FVertexInstanceID VertexInstanceIDs[3];
+
+		/** Polygon which contains this triangle */
+		FPolygonID PolygonID;
+
+	public:
+
+		/** Gets the specified triangle vertex instance ID.  Pass an index between 0 and 2 inclusive. */
+		inline FVertexInstanceID GetVertexInstanceID(const int32 Index) const
+		{
+			checkSlow(Index >= 0 && Index <= 2);
+			return VertexInstanceIDs[Index];
+		}
+
+		/** Sets the specified triangle vertex instance ID.  Pass an index between 0 and 2 inclusive, and the new vertex instance ID to store. */
+		inline void SetVertexInstanceID(const int32 Index, const FVertexInstanceID NewVertexInstanceID)
+		{
+			// When we deprecate direct member access, this will be a simple array lookup
+			checkSlow(Index >= 0 && Index <= 2);
+			VertexInstanceIDs[Index] = NewVertexInstanceID;
+		}
+
+		/** Serializer */
+		friend FArchive& operator<<(FArchive& Ar, FMeshTriangle& Triangle)
+		{
+			Ar << Triangle.VertexInstanceIDs[0];
+			Ar << Triangle.VertexInstanceIDs[1];
+			Ar << Triangle.VertexInstanceIDs[2];
+
+			if (!Ar.IsLoading() || Ar.CustomVer(FEditorObjectVersion::GUID) >= FEditorObjectVersion::MeshDescriptionTriangles)
+			{
+				Ar << Triangle.PolygonID;
+			}
+
+			return Ar;
+		}
+	};
+
 	/** The rendering section index for this mesh section */
 	UPROPERTY()
 	uint32 RenderingSectionIndex;
diff --git a/Engine/Plugins/Runtime/EditableMesh/Source/EditableMesh/EditableMesh.cpp b/Engine/Plugins/Runtime/EditableMesh/Source/EditableMesh/EditableMesh.cpp
--- a/Engine/Plugins/Runtime/EditableMesh/Source/EditableMesh/EditableMesh.cpp
+++ b/Engine/Plugins/Runtime/EditableMesh/Source/EditableMesh/EditableMesh.cpp
@@ -91,9 +91,9 @@ template <> struct TIsAttributeTypeLerpable<int> { static const bool Value = tru
 
 
 template <typename T, typename ElementIDType>
-static typename TEnableIf<TIsAttributeTypeLerpable<T>::Value>::Type InterpAttributeOfType( FMeshElementAttributeList& AttributeList, const FName AttributeName, TMeshAttributesConstRef<ElementIDType, T> AttributeArrayRef, const ElementIDType ID0, const ElementIDType ID1, const float Alpha )
+static typename TEnableIf<TIsAttributeTypeLerpable<T>::Value>::Type InterpAttributeOfType( FMeshElementAttributeList& AttributeList, const FName AttributeName, TMeshAttributesRef<ElementIDType, const T> AttributeArrayRef, const ElementIDType ID0, const ElementIDType ID1, const float Alpha )
 {
-	for( int32 Index = 0; Index < AttributeArrayRef.GetNumIndices(); ++Index )
+	for( int32 Index = 0; Index < AttributeArrayRef.GetNumChannels(); ++Index )
 	{
 		const T& Value0 = AttributeArrayRef.Get( ID0, Index );
 		const T& Value1 = AttributeArrayRef.Get( ID1, Index );
@@ -107,10 +107,10 @@ static typename TEnableIf<TIsAttributeTypeLerpable<T>::Value>::Type InterpAttrib
 }
 
 template <typename T, typename ElementIDType>
-static typename TEnableIf<!TIsAttributeTypeLerpable<T>::Value>::Type InterpAttributeOfType( FMeshElementAttributeList& AttributeList, const FName AttributeName, TMeshAttributesConstRef<ElementIDType, T> AttributeArrayRef, const ElementIDType ID0, const ElementIDType ID1, const float Alpha )
+static typename TEnableIf<!TIsAttributeTypeLerpable<T>::Value>::Type InterpAttributeOfType( FMeshElementAttributeList& AttributeList, const FName AttributeName, TMeshAttributesRef<ElementIDType, const T> AttributeArrayRef, const ElementIDType ID0, const ElementIDType ID1, const float Alpha )
 {
 	// Non-Lerpable attribute type: just choose whichever value alpha is closest to
-	for( int32 Index = 0; Index < AttributeArrayRef.GetNumIndices(); ++Index )
+	for( int32 Index = 0; Index < AttributeArrayRef.GetNumChannels(); ++Index )
 	{
 		const T& Value0 = AttributeArrayRef.Get( ID0, Index );
 		const T& Value1 = AttributeArrayRef.Get( ID1, Index );
@@ -156,12 +156,12 @@ static void InterpAttributesByPredicate( FMeshElementAttributeList& AttributeLis
 
 
 template <typename T, typename ElementIDType>
-static typename TEnableIf<TIsAttributeTypeLerpable<T>::Value>::Type InterpAttributeOfType( FMeshElementAttributeList& AttributeList, const FName AttributeName, TMeshAttributesConstRef<ElementIDType, T> AttributeArrayRef, const ElementIDType ID0, const ElementIDType ID1, const ElementIDType ID2, const FVector Weights )
+static typename TEnableIf<TIsAttributeTypeLerpable<T>::Value>::Type InterpAttributeOfType( FMeshElementAttributeList& AttributeList, const FName AttributeName, TMeshAttributesRef<ElementIDType, const T> AttributeArrayRef, const ElementIDType ID0, const ElementIDType ID1, const ElementIDType ID2, const FVector Weights )
 {
 	if( EnumHasAllFlags( AttributeArrayRef.GetFlags(), EMeshAttributeFlags::Lerpable ) )
 	{
 		// Get result by interpolating between the three values according to the barycentric weight
-		for( int32 Index = 0; Index < AttributeArrayRef.GetNumIndices(); ++Index )
+		for( int32 Index = 0; Index < AttributeArrayRef.GetNumChannels(); ++Index )
 		{
 			const T& Value0 = AttributeArrayRef.Get( ID0, Index );
 			const T& Value1 = AttributeArrayRef.Get( ID1, Index );
@@ -175,7 +175,7 @@ static typename TEnableIf<TIsAttributeTypeLerpable<T>::Value>::Type InterpAttrib
 		// Non-Lerpable attribute type: just choose whichever value the barycentric weight is closest to
 		const ElementIDType ClosestID = ( Weights.X > Weights.Y ) ? ( ( Weights.X > Weights.Z ) ? ID0 : ID2 ) : ( ( Weights.Y > Weights.Z ) ? ID1 : ID2 );
 
-		for( int32 Index = 0; Index < AttributeArrayRef.GetNumIndices(); ++Index )
+		for( int32 Index = 0; Index < AttributeArrayRef.GetNumChannels(); ++Index )
 		{
 			AttributeList.Attributes.Emplace( AttributeName, Index, FMeshElementAttributeValue( AttributeArrayRef.Get( ClosestID, Index ) ) );
 		}
@@ -183,12 +183,12 @@ static typename TEnableIf<TIsAttributeTypeLerpable<T>::Value>::Type InterpAttrib
 }
 
 template <typename T, typename ElementIDType>
-static typename TEnableIf<!TIsAttributeTypeLerpable<T>::Value>::Type InterpAttributeOfType( FMeshElementAttributeList& AttributeList, const FName AttributeName, TMeshAttributesConstRef<ElementIDType, T> AttributeArrayRef, const ElementIDType ID0, const ElementIDType ID1, const ElementIDType ID2, const FVector Weights )
+static typename TEnableIf<!TIsAttributeTypeLerpable<T>::Value>::Type InterpAttributeOfType( FMeshElementAttributeList& AttributeList, const FName AttributeName, TMeshAttributesRef<ElementIDType, const T> AttributeArrayRef, const ElementIDType ID0, const ElementIDType ID1, const ElementIDType ID2, const FVector Weights )
 {
 	// Non-Lerpable attribute type: just choose whichever value the barycentric weight is closest to
 	const ElementIDType ClosestID = ( Weights.X > Weights.Y ) ? ( ( Weights.X > Weights.Z ) ? ID0 : ID2 ) : ( ( Weights.Y > Weights.Z ) ? ID1 : ID2 );
 
-	for( int32 Index = 0; Index < AttributeArrayRef.GetNumIndices(); ++Index )
+	for( int32 Index = 0; Index < AttributeArrayRef.GetNumChannels(); ++Index )
 	{
 		AttributeList.Attributes.Emplace( AttributeName, Index, FMeshElementAttributeValue( AttributeArrayRef.Get( ClosestID, Index ) ) );
 	}
@@ -248,7 +248,7 @@ static bool CompareAttributesByPredicate( const TAttributesSet<ElementIDType>& A
 		{
 			if( Pred( AttributeName, AttributeArrayRef ) )
 			{
-				for( int32 Index = 0; Index < AttributeArrayRef.GetNumIndices(); ++Index )
+				for( int32 Index = 0; Index < AttributeArrayRef.GetNumChannels(); ++Index )
 				{
 					if( !AreAttributeValuesNearlyEqual( AttributeArrayRef.Get( ID0, Index ), AttributeArrayRef.Get( ID1, Index ) ) )
 					{
@@ -275,7 +275,7 @@ static void CopyAllAttributes( TAttributesSet<ElementIDType>& DestAttributesSet,
 	SrcAttributesSet.ForEachAttributeIndicesArray(
 		[ &DestAttributesSet, ElementID ]( const FName AttributeName, auto AttributeArrayRef )
 		{
-			for( int32 Index = 0; Index < AttributeArrayRef.GetNumIndices(); ++Index )
+			for( int32 Index = 0; Index < AttributeArrayRef.GetNumChannels(); ++Index )
 			{
 				DestAttributesSet.SetAttribute( ElementID, AttributeName, Index, AttributeArrayRef.Get( ElementID, Index ) );
 			}
@@ -291,7 +291,7 @@ static void BackupAllAttributes( FMeshElementAttributeList& AttributeList, const
 	AttributesSet.ForEach(
 		[ &AttributeList, ElementID ]( const FName AttributeName, auto AttributeArrayRef )
 		{
-			for( int32 Index = 0; Index < AttributeArrayRef.GetNumIndices(); ++Index )
+			for( int32 Index = 0; Index < AttributeArrayRef.GetNumChannels(); ++Index )
 			{
 				const auto& Value = AttributeArrayRef.Get( ElementID, Index );
 				AttributeList.Attributes.Emplace( AttributeName, Index, FMeshElementAttributeValue( Value ) );
@@ -310,7 +310,7 @@ static void BackupAttributesByPredicate( FMeshElementAttributeList& AttributeLis
 		{
 			if( Pred( AttributeName, AttributeArrayRef ) )
 			{
-				for( int32 Index = 0; Index < AttributeArrayRef.GetNumIndices(); ++Index )
+				for( int32 Index = 0; Index < AttributeArrayRef.GetNumChannels(); ++Index )
 				{
 					const auto& Value = AttributeArrayRef.Get( ElementID, Index );
 					AttributeList.Attributes.Emplace( AttributeName, Index, FMeshElementAttributeValue( Value ) );
@@ -886,8 +886,8 @@ int32 UEditableMesh::GetVertexConnectedEdgeCount( const FVertexID VertexID ) con
 
 FEdgeID UEditableMesh::GetVertexConnectedEdge( const FVertexID VertexID, const int32 ConnectedEdgeNumber ) const
 {
-	const TArray<FEdgeID>& ConnectedEdgeIDs = GetMeshDescription()->GetVertexConnectedEdges( VertexID );
-	return ( ConnectedEdgeNumber >= 0 && ConnectedEdgeNumber < ConnectedEdgeIDs.Num() ) ? ConnectedEdgeIDs[ ConnectedEdgeNumber ] : FEdgeID::Invalid;
+	TArrayView<const FEdgeID> ConnectedEdgeIDs = GetMeshDescription()->GetVertexConnectedEdges( VertexID );
+	return ( ConnectedEdgeNumber >= 0 && ConnectedEdgeNumber < ConnectedEdgeIDs.Num() ) ? ConnectedEdgeIDs[ ConnectedEdgeNumber ] : INDEX_NONE;
 }
 
 
@@ -912,7 +912,7 @@ int32 UEditableMesh::GetVertexInstanceConnectedPolygonCount( const FVertexInstan
 FPolygonID UEditableMesh::GetVertexInstanceConnectedPolygon( const FVertexInstanceID VertexInstanceID, const int32 ConnectedPolygonNumber ) const
 {
 	const TArray<FPolygonID>& ConnectedPolygonIDs = GetMeshDescription()->GetVertexInstanceConnectedPolygons( VertexInstanceID );
-	return ( ConnectedPolygonNumber >= 0 && ConnectedPolygonNumber < ConnectedPolygonIDs.Num() ) ? ConnectedPolygonIDs[ ConnectedPolygonNumber ] : FPolygonID::Invalid;
+	return ( ConnectedPolygonNumber >= 0 && ConnectedPolygonNumber < ConnectedPolygonIDs.Num() ) ? ConnectedPolygonIDs[ ConnectedPolygonNumber ] : INDEX_NONE;
 }
 
 
@@ -944,7 +944,7 @@ int32 UEditableMesh::GetEdgeConnectedPolygonCount( const FEdgeID EdgeID ) const
 FPolygonID UEditableMesh::GetEdgeConnectedPolygon( const FEdgeID EdgeID, const int32 ConnectedPolygonNumber ) const
 {
 	const TArray<FPolygonID>& ConnectedPolygonIDs = GetMeshDescription()->GetEdgeConnectedPolygons( EdgeID );
-	return ( ConnectedPolygonNumber >= 0 && ConnectedPolygonNumber < ConnectedPolygonIDs.Num() ) ? ConnectedPolygonIDs[ ConnectedPolygonNumber ] : FPolygonID::Invalid;
+	return ( ConnectedPolygonNumber >= 0 && ConnectedPolygonNumber < ConnectedPolygonIDs.Num() ) ? ConnectedPolygonIDs[ ConnectedPolygonNumber ] : INDEX_NONE;
 }
 
 
@@ -968,8 +968,8 @@ int32 UEditableMesh::GetPolygonCountInGroup( const FPolygonGroupID PolygonGroupI
 
 FPolygonID UEditableMesh::GetPolygonInGroup( const FPolygonGroupID PolygonGroupID, const int32 PolygonNumber ) const
 {
-	const TArray<FPolygonID>& PolygonIDs = GetMeshDescription()->GetPolygonGroupPolygons( PolygonGroupID );
-	return ( PolygonNumber >= 0 && PolygonNumber < PolygonIDs.Num() ) ? PolygonIDs[ PolygonNumber ] : FPolygonID::Invalid;
+	TArrayView<const FPolygonID> PolygonIDs = GetMeshDescription()->GetPolygonGroupPolygons( PolygonGroupID );
+	return ( PolygonNumber >= 0 && PolygonNumber < PolygonIDs.Num() ) ? PolygonIDs[ PolygonNumber ] : INDEX_NONE;
 }
 
 
@@ -999,15 +999,15 @@ int32 UEditableMesh::GetPolygonPerimeterVertexCount( const FPolygonID PolygonID
 
 FVertexInstanceID UEditableMesh::GetPolygonPerimeterVertexInstance( const FPolygonID PolygonID, const int32 PolygonVertexNumber ) const
 {
-	const TArray<FVertexInstanceID>& VertexInstanceIDs = GetMeshDescription()->GetPolygonVertexInstances( PolygonID );
-	return ( PolygonVertexNumber >= 0 && PolygonVertexNumber < VertexInstanceIDs.Num() ) ? VertexInstanceIDs[ PolygonVertexNumber ] : FVertexInstanceID::Invalid;
+	TArray<FVertexInstanceID> VertexInstanceIDs = GetMeshDescription()->GetPolygonVertexInstances( PolygonID );
+	return ( PolygonVertexNumber >= 0 && PolygonVertexNumber < VertexInstanceIDs.Num() ) ? VertexInstanceIDs[ PolygonVertexNumber ] : INDEX_NONE;
 }
 
 
 FVertexID UEditableMesh::GetPolygonPerimeterVertex( const FPolygonID PolygonID, const int32 PolygonVertexNumber ) const
 {
 	const FVertexInstanceID VertexInstanceID = GetPolygonPerimeterVertexInstance( PolygonID, PolygonVertexNumber );
-	return ( VertexInstanceID != FVertexInstanceID::Invalid ) ? GetMeshDescription()->GetVertexInstanceVertex( VertexInstanceID ) : FVertexID::Invalid;
+	return ( VertexInstanceID != INDEX_NONE ) ? GetMeshDescription()->GetVertexInstanceVertex( VertexInstanceID ) : INDEX_NONE;
 }
 
 
@@ -1019,8 +1019,8 @@ int32 UEditableMesh::GetPolygonTriangulatedTriangleCount( const FPolygonID Polyg
 
 FTriangleID UEditableMesh::GetPolygonTriangulatedTriangle( const FPolygonID PolygonID, int32 PolygonTriangleNumber ) const
 {
-	const TArray<FTriangleID>& Triangles = GetMeshDescription()->GetPolygonTriangleIDs( PolygonID );
-	return ( PolygonTriangleNumber >= 0 && PolygonTriangleNumber < Triangles.Num() ) ? Triangles[ PolygonTriangleNumber ] : FTriangleID::Invalid;
+	TArrayView<const FTriangleID> Triangles = GetMeshDescription()->GetPolygonTriangleIDs( PolygonID );
+	return ( PolygonTriangleNumber >= 0 && PolygonTriangleNumber < Triangles.Num() ) ? Triangles[ PolygonTriangleNumber ] : INDEX_NONE;
 }
 
 
@@ -1328,7 +1328,7 @@ void UEditableMesh::MergeVertexInstances()
 		GetConnectedSoftEdges( VertexID, VertexConnectedSoftEdges );
 
 		// Get all vertex instances of this vertex...
-		const TArray<FVertexInstanceID>& VertexInstanceIDs = GetMeshDescription()->GetVertexVertexInstances( VertexID );
+		TArrayView<const FVertexInstanceID> VertexInstanceIDs = GetMeshDescription()->GetVertexVertexInstances( VertexID );
 
 		// ...and iterate through pairs of vertex instances, looking for potential to merge them
 		for( int32 IndexA = 0; IndexA < VertexInstanceIDs.Num() - 1; ++IndexA )
@@ -1472,7 +1472,7 @@ void UEditableMesh::GetEdgeLoopElements( const FEdgeID EdgeID, TArray<FEdgeID>&
 	// We also store which vertex of the edge has already been processed (so we don't retrace our steps when processing stack items).
 	static TArray<TTuple<FEdgeID, FVertexID>> EdgeStack;
 	EdgeStack.Reset();
-	EdgeStack.Push( MakeTuple( EdgeID, FVertexID::Invalid ) );
+	EdgeStack.Push( MakeTuple( EdgeID, FVertexID(INDEX_NONE) ) );
 
 	// Process edge IDs on the stack
 	while( EdgeStack.Num() > 0 )
@@ -1499,7 +1499,7 @@ void UEditableMesh::GetEdgeLoopElements( const FEdgeID EdgeID, TArray<FEdgeID>&
 			}
 
 			// This is the candidate edge ID which continues the loop beyond the vertex being processed
-			FEdgeID AdjacentEdgeID = FEdgeID::Invalid;
+			FEdgeID AdjacentEdgeID = INDEX_NONE;
 
 			// Iterate through all edges connected to this vertex
 			for( const FEdgeID ConnectedEdgeID : GetMeshDescription()->GetVertexConnectedEdges( ConnectedVertexID ) )
@@ -1521,7 +1521,7 @@ void UEditableMesh::GetEdgeLoopElements( const FEdgeID EdgeID, TArray<FEdgeID>&
 					// We have found an edge connected to this vertex which doesn't share any polys with the original edge
 					if( bIsCandidateEdge )
 					{
-						if( AdjacentEdgeID == FEdgeID::Invalid )
+						if( AdjacentEdgeID == INDEX_NONE )
 						{
 							// If it's the first such edge which meets the criteria, remember it
 							AdjacentEdgeID = ConnectedEdgeID;
@@ -1529,14 +1529,14 @@ void UEditableMesh::GetEdgeLoopElements( const FEdgeID EdgeID, TArray<FEdgeID>&
 						else
 						{
 							// If we already have a possible edge, stop the loop here; we don't allow splits in the loop if there is more than one candidate
-							AdjacentEdgeID = FEdgeID::Invalid;
+							AdjacentEdgeID = INDEX_NONE;
 							break;
 						}
 					}
 				}
 			}
 
-			if( AdjacentEdgeID != FEdgeID::Invalid )
+			if( AdjacentEdgeID != INDEX_NONE )
 			{
 				EdgeStack.Push( MakeTuple( AdjacentEdgeID, ConnectedVertexID ) );
 			}
@@ -1572,10 +1572,10 @@ void UEditableMesh::GetPolygonPerimeterVertexInstances( const FPolygonID Polygon
 FEdgeID UEditableMesh::GetPolygonPerimeterEdge( const FPolygonID PolygonID, const int32 PerimeterEdgeNumber, bool& bOutEdgeWindingIsReversedForPolygon ) const
 {
 	const FMeshDescription* Description = GetMeshDescription();
-	const TArray<FVertexInstanceID>& VertexInstanceIDs = Description->GetPolygonVertexInstances( PolygonID );
+	TArray<FVertexInstanceID> VertexInstanceIDs = Description->GetPolygonVertexInstances( PolygonID );
 	if( PerimeterEdgeNumber >= VertexInstanceIDs.Num() )
 	{
-		return FEdgeID::Invalid;
+		return INDEX_NONE;
 	}
 
 	const FVertexID VertexID0 = Description->GetVertexInstanceVertex( VertexInstanceIDs[ PerimeterEdgeNumber ] );
@@ -1587,7 +1587,7 @@ FEdgeID UEditableMesh::GetPolygonPerimeterEdge( const FPolygonID PolygonID, cons
 FEdgeID UEditableMesh::GetVertexPairEdge( const FVertexID StartVertexID, const FVertexID EndVertexID, bool& bOutEdgeWindingIsReversed ) const
 {
 	const FEdgeID EdgeID = GetMeshDescription()->GetVertexPairEdge( StartVertexID, EndVertexID );
-	if( EdgeID != FEdgeID::Invalid )
+	if( EdgeID != INDEX_NONE )
 	{
 		// @note: seems like a reversal of logic here, but everything else expects this behavior now
 		bOutEdgeWindingIsReversed = ( GetMeshDescription()->GetEdgeVertex( EdgeID, 0 ) == StartVertexID );
@@ -1600,7 +1600,7 @@ FEdgeID UEditableMesh::GetVertexPairEdge( const FVertexID StartVertexID, const F
 void UEditableMesh::GetPolygonPerimeterEdges( const FPolygonID PolygonID, TArray<FEdgeID>& OutPolygonPerimeterEdgeIDs ) const
 {
 	const FMeshDescription* Description = GetMeshDescription();
-	const TArray<FVertexInstanceID>& VertexInstanceIDs = Description->GetPolygonVertexInstances( PolygonID );
+	TArray<FVertexInstanceID> VertexInstanceIDs = Description->GetPolygonVertexInstances( PolygonID );
 	const int32 NumContourEdges = VertexInstanceIDs.Num();
 
 	OutPolygonPerimeterEdgeIDs.SetNumUninitialized( NumContourEdges, false );
@@ -1758,42 +1758,12 @@ void UEditableMesh::RefreshOpenSubdiv()
 				{
 					OsdCornerVertexIndices.Reset();
 					OsdCornerWeights.Reset();
-
-					const TVertexAttributesRef<float> VertexCornerSharpnesses = GetMeshDescription()->VertexAttributes().GetAttributesRef<float>( MeshAttribute::Vertex::CornerSharpness );
-
-					for( const FVertexID VertexID : GetMeshDescription()->Vertices().GetElementIDs() )
-					{
-						const float VertexCornerSharpness = VertexCornerSharpnesses[ VertexID ];
-						if( VertexCornerSharpness > SMALL_NUMBER )
-						{
-							// This vertex is (at least partially) a subdivision corner
-							OsdCornerVertexIndices.Add( VertexID.GetValue() );
-							OsdCornerWeights.Add( OpenSubdivCreaseWeightMultiplier * VertexCornerSharpness );
-						}
-					}
 				}
 
 				// Edge creases
 				{
 					OsdCreaseVertexIndexPairs.Reset();
 					OsdCreaseWeights.Reset();
-
-					const TEdgeAttributesRef<float> EdgeCreaseSharpnesses = GetMeshDescription()->EdgeAttributes().GetAttributesRef<float>( MeshAttribute::Edge::CreaseSharpness );
-
-					for( const FEdgeID EdgeID : GetMeshDescription()->Edges().GetElementIDs() )
-					{
-						const float EdgeCreaseSharpness = EdgeCreaseSharpnesses[ EdgeID ];
-						if( EdgeCreaseSharpness > SMALL_NUMBER )
-						{
-							// This edge is (at least partially) creased
-							FVertexID EdgeVertexID0, EdgeVertexID1;
-							GetEdgeVertices( EdgeID, /* Out */ EdgeVertexID0, /* Out */ EdgeVertexID1 );
-
-							OsdCreaseVertexIndexPairs.Add( EdgeVertexID0.GetValue() );
-							OsdCreaseVertexIndexPairs.Add( EdgeVertexID1.GetValue() );
-							OsdCreaseWeights.Add( OpenSubdivCreaseWeightMultiplier * EdgeCreaseSharpness );
-						}
-					}
 				}
 
 				OsdNumVerticesPerFace.SetNum( OsdTopologyDescriptor.numFaces, false );
@@ -2006,12 +1976,13 @@ void UEditableMesh::GenerateOpenSubdivLimitSurfaceData()
 				{
 					FirstPolygonNumberForPolygonGroups.Add( NumPolygonsSoFar );
 
-					const TArray<FPolygonID>& PolygonIDs = GetMeshDescription()->GetPolygonGroupPolygons( PolygonGroupID );
+					TArrayView<const FPolygonID> PolygonIDs = GetMeshDescription()->GetPolygonGroupPolygons( PolygonGroupID );
 					NumPolygonsSoFar += PolygonIDs.Num();
 
 					for( const FPolygonID PolygonID : PolygonIDs )
 					{
-						for( const FVertexInstanceID VertexInstanceID : GetMeshDescription()->GetPolygonVertexInstances( PolygonID ) )
+						TArray<FVertexInstanceID> VertexInstanceIDs = GetMeshDescription()->GetPolygonVertexInstances( PolygonID );
+						for( const FVertexInstanceID VertexInstanceID : VertexInstanceIDs )
 						{
 							FOsdFVarVertexData& FVarVertexData = *new( FVarVertexDatas ) FOsdFVarVertexData();
 							FVarVertexData.TextureCoordinates[ 0 ] = TextureCoordinateCount > 0 ? VertexUVs.Get( VertexInstanceID, 0 ) : FVector2D::ZeroVector;
@@ -2305,7 +2276,7 @@ void UEditableMesh::GenerateOpenSubdivLimitSurfaceData()
 				SubdividedWireEdge.EdgeVertex1PositionIndex = OsdLimitEdgeVertices[ 1 ];
 
 				// Default to not highlighting this edge as a base cage counterpart.  We'll actually figure this out below.
-				SubdividedWireEdge.CounterpartEdgeID = FEdgeID::Invalid;
+				SubdividedWireEdge.CounterpartEdgeID = INDEX_NONE;
 			}
 
 			{
@@ -2419,7 +2390,7 @@ FTriangleID UEditableMesh::ComputeBarycentricWeightForPointOnPolygon( const FPol
 		}
 	}
 
-	return FTriangleID::Invalid;
+	return INDEX_NONE;
 }
 
 
@@ -2512,7 +2483,7 @@ void UEditableMesh::CreateMissingPolygonPerimeterEdges( const FPolygonID Polygon
 		const FVertexID NextVertexID = GetPolygonPerimeterVertex( PolygonID, ( PerimeterVertexNumber + 1 ) % NumPolygonPerimeterVertices );
 
 		// Find the edge that connects these vertices
-		FEdgeID FoundEdgeID = FEdgeID::Invalid;
+		FEdgeID FoundEdgeID = INDEX_NONE;
 		bool bFoundEdge = false;
 
 		const int32 NumVertexConnectedEdges = GetVertexConnectedEdgeCount( VertexID );	// @todo mesheditor perf: Could be made faster by always iterating over the vertex with the fewest number of connected edges, instead of the first edge vertex
@@ -2639,7 +2610,7 @@ void UEditableMesh::SplitEdge( const FEdgeID EdgeID, const TArray<float>& Splits
 		// Iterate through each polygon connected to the edge being split
 		for( const FPolygonID ConnectedPolygonID : GetMeshDescription()->GetEdgeConnectedPolygons( EdgeID ) )
 		{
-			const TArray<FVertexInstanceID>& PerimeterVertexInstances = GetMeshDescription()->GetPolygonVertexInstances( ConnectedPolygonID );
+			TArray<FVertexInstanceID> PerimeterVertexInstances = GetMeshDescription()->GetPolygonVertexInstances( ConnectedPolygonID );
 			const int32 NumVertexInstances = PerimeterVertexInstances.Num();
 
 			// Find vertex indices around the perimeter corresponding to the vertices defining the edge
@@ -2879,10 +2850,10 @@ void UEditableMesh::FindPolygonLoop( const FEdgeID EdgeID, TArray<FEdgeID>& OutE
 		const FVector CurrentEdgeVertex1 = VertexPositions[ CurrentEdgeVertexIDs[ 1 ] ];
 		const FVector CurrentEdgeDirection = ( CurrentEdgeVertex1 - CurrentEdgeVertex0 ).GetSafeNormal();
 
-		const FEdgeID NextEdgeIDInPath = OutReversedEdgeIDPathToTake.Num() > 0 ? OutReversedEdgeIDPathToTake.Pop( false ) : FEdgeID::Invalid;
+		const FEdgeID NextEdgeIDInPath = OutReversedEdgeIDPathToTake.Num() > 0 ? OutReversedEdgeIDPathToTake.Pop( false ) : INDEX_NONE;
 
-		FEdgeID BestEdgeID = FEdgeID::Invalid;
-		FPolygonID BestEdgeSplitsPolygon = FPolygonID::Invalid;
+		FEdgeID BestEdgeID = INDEX_NONE;
+		FPolygonID BestEdgeSplitsPolygon = INDEX_NONE;
 		bool bBestEdgeIsInOppositeDirection = false;
 		bool bBestEdgeIsBorderEdge = false;
 		float LargestAbsDotProduct = -1.0f;
@@ -2907,7 +2878,7 @@ void UEditableMesh::FindPolygonLoop( const FEdgeID EdgeID, TArray<FEdgeID>& OutE
 					if( CandidateEdgeID != CurrentEdgeID )
 					{
 						// If we need to follow a specific path, then do that
-						if( NextEdgeIDInPath == FEdgeID::Invalid || CandidateEdgeID == NextEdgeIDInPath )
+						if( NextEdgeIDInPath == INDEX_NONE || CandidateEdgeID == NextEdgeIDInPath )
 						{
 							FVertexID CandidateEdgeVertexIDs[ 2 ];
 							GetEdgeVertices( CandidateEdgeID, /* Out */ CandidateEdgeVertexIDs[ 0 ], /* Out */ CandidateEdgeVertexIDs[ 1 ] );
@@ -2943,7 +2914,7 @@ void UEditableMesh::FindPolygonLoop( const FEdgeID EdgeID, TArray<FEdgeID>& OutE
 											FMath::Lerp( CurrentEdgeVertex0, CurrentEdgeVertex1, 0.5f ) ).GetSafeNormal();
 									const float CandidateEdgeDot = FVector::DotProduct( DirectionTowardCenterOfCurrentEdge, DirectionTowardCenterOfCandidateEdge );
 
-									check( BestEdgeID != FEdgeID::Invalid );
+									check( BestEdgeID != INDEX_NONE );
 
 									FVertexID BestEdgeVertexIDs[ 2 ];
 									GetEdgeVertices( BestEdgeID, /* Out */ BestEdgeVertexIDs[ 0 ], /* Out */ BestEdgeVertexIDs[ 1 ] );
@@ -2973,7 +2944,7 @@ void UEditableMesh::FindPolygonLoop( const FEdgeID EdgeID, TArray<FEdgeID>& OutE
 									// multiple colinear edges on the same polygon (such as after SplitEdge() is called to insert a vertex on a polygon.)
 									const float BestEdgeDistance = [this, BestEdgeID, CurrentEdgeVertex0, CurrentEdgeVertex1, &VertexPositions]() -> float
 									{
-										check( BestEdgeID != FEdgeID::Invalid );
+										check( BestEdgeID != INDEX_NONE );
 
 										FVertexID BestEdgeVertexIDs[ 2 ];
 										GetEdgeVertices( BestEdgeID, /* Out */ BestEdgeVertexIDs[ 0 ], /* Out */ BestEdgeVertexIDs[ 1 ] );
@@ -3018,7 +2989,7 @@ void UEditableMesh::FindPolygonLoop( const FEdgeID EdgeID, TArray<FEdgeID>& OutE
 			}
 		}
 
-		if( BestEdgeID != FEdgeID::Invalid &&
+		if( BestEdgeID != INDEX_NONE &&
 			!OutPolygonIDsToSplit.Contains( BestEdgeSplitsPolygon ) )	// If we try to re-split the same polygon twice, then this loop is not valid
 		{
 			// OK, this polygon will definitely be split
@@ -3359,7 +3330,7 @@ void UEditableMesh::SplitPolygons( const TArray<FPolygonToSplit>& PolygonsToSpli
 			for( const int32 VertexNumber : PerimeterVertexNumbers )
 			{
 				FVertexAndAttributes& NewVertexAndAttributes = *new( NewPolygon.PerimeterVertices ) FVertexAndAttributes();
-				NewVertexAndAttributes.VertexID = FVertexID::Invalid;
+				NewVertexAndAttributes.VertexID = INDEX_NONE;
 				NewVertexAndAttributes.VertexInstanceID = GetPolygonPerimeterVertexInstance( PolygonID, VertexNumber );
 			}
 
@@ -3672,7 +3643,7 @@ void UEditableMesh::CreateVertices( const TArray<FVertexToCreate>& VerticesToCre
 		for( const FVertexToCreate& VertexToCreate : VerticesToCreate )
 		{
 			FVertexID VertexID = VertexToCreate.OriginalVertexID;
-			if( VertexID != FVertexID::Invalid )
+			if( VertexID != INDEX_NONE )
 			{
 				GetMeshDescription()->CreateVertexWithID( VertexID );
 			}
@@ -3728,7 +3699,7 @@ void UEditableMesh::CreateVertexInstances( const TArray<FVertexInstanceToCreate>
 		for( const FVertexInstanceToCreate& VertexInstanceToCreate : VertexInstancesToCreate )
 		{
 			FVertexInstanceID VertexInstanceID = VertexInstanceToCreate.OriginalVertexInstanceID;
-			if( VertexInstanceID != FVertexInstanceID::Invalid )
+			if( VertexInstanceID != INDEX_NONE )
 			{
 				GetMeshDescription()->CreateVertexInstanceWithID( VertexInstanceID, VertexInstanceToCreate.VertexID );
 			}
@@ -3786,7 +3757,7 @@ void UEditableMesh::CreateEdges( const TArray<FEdgeToCreate>& EdgesToCreate, TAr
 		for( const FEdgeToCreate& EdgeToCreate : EdgesToCreate )
 		{
 			FEdgeID EdgeID = EdgeToCreate.OriginalEdgeID;
-			if( EdgeID != FEdgeID::Invalid )
+			if( EdgeID != INDEX_NONE )
 			{
 				GetMeshDescription()->CreateEdgeWithID( EdgeID, EdgeToCreate.VertexID0, EdgeToCreate.VertexID1 );
 			}
@@ -3834,10 +3805,10 @@ FVertexInstanceID UEditableMesh::CreateVertexInstanceForContourVertex( const FVe
 {
 	FVertexInstanceID NewVertexInstanceID = ContourVertex.VertexInstanceID;
 
-	if( NewVertexInstanceID == FVertexInstanceID::Invalid )
+	if( NewVertexInstanceID == INDEX_NONE )
 	{
 		const FVertexID VertexID = ContourVertex.VertexID;
-		check( VertexID != FVertexID::Invalid );
+		check( VertexID != INDEX_NONE );
 
 		// For now, always create a new vertex instance.
 		// @todo mesheditor: revisit splitting/welding identical vertex instances.
@@ -3859,7 +3830,7 @@ FVertexInstanceID UEditableMesh::CreateVertexInstanceForContourVertex( const FVe
 	else
 	{
 		// Cannot specify both a vertex instance ID and a vertex ID
-		check( ContourVertex.VertexID == FVertexID::Invalid );
+		check( ContourVertex.VertexID == INDEX_NONE );
 	}
 
 	return NewVertexInstanceID;
@@ -3887,23 +3858,23 @@ void UEditableMesh::CreatePolygonContour( const TArray<FVertexAndAttributes>& Co
 		FVertexID VertexID1 = Contour[ NextVertexNumber ].VertexID;
 
 		// If valid vertex IDs weren't supplied, we'll deduce it from the vertex instance IDs instead (which must be valid instead).
-		if( VertexID0 == FVertexID::Invalid )
+		if( VertexID0 == INDEX_NONE )
 		{
 			const FVertexInstanceID VertexInstanceID0 = Contour[ VertexNumber ].VertexInstanceID;
-			check( VertexInstanceID0 != FVertexInstanceID::Invalid );
+			check( VertexInstanceID0 != INDEX_NONE );
 			VertexID0 = GetVertexInstanceVertex( VertexInstanceID0 );
 		}
 
-		if( VertexID1 == FVertexID::Invalid )
+		if( VertexID1 == INDEX_NONE )
 		{
 			const FVertexInstanceID VertexInstanceID1 = Contour[ NextVertexNumber ].VertexInstanceID;
-			check( VertexInstanceID1 != FVertexInstanceID::Invalid );
+			check( VertexInstanceID1 != INDEX_NONE );
 			VertexID1 = GetVertexInstanceVertex( VertexInstanceID1 );
 		}
 
 		FEdgeID EdgeID = GetMeshDescription()->GetVertexPairEdge( VertexID0, VertexID1 );
 
-		if( EdgeID == FEdgeID::Invalid )
+		if( EdgeID == INDEX_NONE )
 		{
 			// Create the new edge. Don't specify any attributes now.
 			EdgesToCreate.Emplace();
@@ -3932,10 +3903,10 @@ void UEditableMesh::CreatePolygonContour( const TArray<FVertexAndAttributes>& Co
 		// If it is not valid, copy the invalid ID. They will be replaced by the IDs of the new vertex instances in a later pass.
 		OutVertexInstanceIDs[ VertexNumber ] = Contour[ VertexNumber ].VertexInstanceID;
 
-		if( Contour[ VertexNumber ].VertexInstanceID != FVertexInstanceID::Invalid )
+		if( Contour[ VertexNumber ].VertexInstanceID != INDEX_NONE )
 		{
 			// Cannot specify both a vertex instance ID and a vertex ID
-			check( Contour[ VertexNumber ].VertexID == FVertexID::Invalid );
+			check( Contour[ VertexNumber ].VertexID == INDEX_NONE );
 
 			// Vertex instance ID was specified; nothing more to do, it has already been output directly
 			// @todo editable mesh: THIS WILL BREAK if you specify a vertex instance which is being used by a polygon the other side of a hard edge.
@@ -3945,7 +3916,7 @@ void UEditableMesh::CreatePolygonContour( const TArray<FVertexAndAttributes>& Co
 		{
 			// We need to create a vertex instance; check that the vertex ID is valid
 			const FVertexID VertexID = Contour[ VertexNumber ].VertexID;
-			check( VertexID != FVertexID::Invalid );
+			check( VertexID != INDEX_NONE );
 
 			// For now, always create a new vertex instance. It will be merged with an existing one later if appropriate.
 			VertexInstancesToCreate.Emplace();
@@ -3971,7 +3942,7 @@ void UEditableMesh::CreatePolygonContour( const TArray<FVertexAndAttributes>& Co
 
 		for( FVertexInstanceID& VertexInstanceID : OutVertexInstanceIDs )
 		{
-			if( VertexInstanceID == FVertexInstanceID::Invalid )
+			if( VertexInstanceID == INDEX_NONE )
 			{
 				VertexInstanceID = NewVertexInstanceIDs[ NewVertexInstanceIndex ];
 				NewVertexInstanceIndex++;
@@ -4003,7 +3974,7 @@ void UEditableMesh::CreatePolygons( const TArray<FPolygonToCreate>& PolygonsToCr
 
 			TArray<FEdgeID> CreatedEdgeIDs;
 			FPolygonID PolygonID = PolygonToCreate.OriginalPolygonID;
-			if( PolygonID != FPolygonID::Invalid )
+			if( PolygonID != INDEX_NONE )
 			{
 				GetMeshDescription()->CreatePolygonWithID( PolygonID, PolygonToCreate.PolygonGroupID, PerimeterVertexInstances, &CreatedEdgeIDs );
 			}
@@ -4039,7 +4010,7 @@ void UEditableMesh::CreatePolygons( const TArray<FPolygonToCreate>& PolygonsToCr
 					const FVertexID VertexID0 = GetMeshDescription()->GetVertexInstanceVertex( VertexInstanceID );
 					const FVertexID VertexID1 = GetMeshDescription()->GetVertexInstanceVertex( LastVertexInstanceID );
 					const FEdgeID EdgeID = GetMeshDescription()->GetVertexPairEdge( VertexID0, VertexID1 );
-					check( EdgeID != FEdgeID::Invalid );
+					check( EdgeID != INDEX_NONE );
 					if( !NewEdgeIDsForContour.Contains( EdgeID ) )
 					{
 						AttributesForEdges.Emplace();
@@ -4126,15 +4097,16 @@ void UEditableMesh::DeletePolygons( const TArray<FPolygonID>& PolygonIDsToDelete
 			PolygonToCreate.PolygonGroupID = GetMeshDescription()->GetPolygonPolygonGroup( PolygonID );
 			PolygonToCreate.OriginalPolygonID = PolygonID;
 
-			PolygonToCreate.PerimeterVertices.Reserve( GetMeshDescription()->GetPolygonVertexInstances( PolygonID ).Num() );
-			for( const FVertexInstanceID VertexInstanceID : GetMeshDescription()->GetPolygonVertexInstances( PolygonID ) )
+			TArray<FVertexInstanceID> VertexInstanceIDs = GetMeshDescription()->GetPolygonVertexInstances( PolygonID );
+			PolygonToCreate.PerimeterVertices.Reserve( VertexInstanceIDs.Num() );
+			for( const FVertexInstanceID VertexInstanceID : VertexInstanceIDs )
 			{
 				PolygonToCreate.PerimeterVertices.Emplace();
 				FVertexAndAttributes& VertexAndAttributes = PolygonToCreate.PerimeterVertices.Last();
 
 				// We rely on undoing recreating vertex instances, therefore we only need pass their IDs.
 				VertexAndAttributes.VertexInstanceID = VertexInstanceID;
-				VertexAndAttributes.VertexID = FVertexID::Invalid;
+				VertexAndAttributes.VertexID = INDEX_NONE;
 			}
 
 		}
@@ -4231,7 +4203,7 @@ void UEditableMesh::CreatePolygonGroups( const TArray<FPolygonGroupToCreate>& Po
 		{
 			// Allocate polygon group
 			FPolygonGroupID PolygonGroupID = PolygonGroupToCreate.OriginalPolygonGroupID;
-			if( PolygonGroupID != FPolygonGroupID::Invalid )
+			if( PolygonGroupID != INDEX_NONE )
 			{
 				GetMeshDescription()->CreatePolygonGroupWithID( PolygonGroupID );
 			}
@@ -4293,7 +4265,6 @@ void UEditableMesh::DeletePolygonGroups( const TArray<FPolygonGroupID>& PolygonG
 		for( int32 Index = PolygonGroupIDs.Num() - 1; Index >= 0; --Index )
 		{
 			const FPolygonGroupID PolygonGroupID = PolygonGroupIDs[ Index ];
-			const FMeshPolygonGroup& PolygonGroup = PolygonGroups[ PolygonGroupID ];
 
 			RevertInput.PolygonGroupsToCreate.Emplace();
 			FPolygonGroupToCreate& PolygonGroupToCreate = RevertInput.PolygonGroupsToCreate.Last();
@@ -4458,6 +4429,7 @@ void UEditableMesh::ChangePolygonsVertexInstances( const TArray<FChangeVertexIns
 
 	// Perform action
 	{
+
 		for( const FChangeVertexInstancesForPolygon& VertexInstancesForPolygon : VertexInstancesForPolygons )
 		{
 			const FPolygonID PolygonID = VertexInstancesForPolygon.PolygonID;
@@ -4490,7 +4462,7 @@ FVertexInstanceID UEditableMesh::GetVertexInstanceInPolygonForVertex( const FPol
 		}
 	}
 
-	return FVertexInstanceID::Invalid;
+	return INDEX_NONE;
 }
 
 
@@ -4526,7 +4498,7 @@ void UEditableMesh::SetPolygonsVertexAttributes( const TArray<FVertexAttributesF
 
 void UEditableMesh::SetPolygonContourVertexAttributes( const FPolygonID PolygonID, const TArray<FMeshElementAttributeList>& AttributeLists )
 {
-	const TArray<FVertexInstanceID>& ContourVertexInstanceIDs = GetMeshDescription()->GetPolygonVertexInstances( PolygonID );
+	TArray<FVertexInstanceID> ContourVertexInstanceIDs = GetMeshDescription()->GetPolygonVertexInstances( PolygonID );
 
 	const int32 NumContourVertices = ContourVertexInstanceIDs.Num();
 	check( AttributeLists.Num() == NumContourVertices );
@@ -4601,7 +4573,7 @@ void UEditableMesh::TryToRemovePolygonEdge( const FEdgeID EdgeID, bool& bOutWasE
 	EM_ENTER( TEXT( "TryToRemovePolygonEdge: %s" ), *EdgeID.ToString() );
 
 	bOutWasEdgeRemoved = false;
-	OutNewPolygonID = FPolygonID::Invalid;
+	OutNewPolygonID = INDEX_NONE;
 
 	// If the edge is not shared by at least two polygons, we can't remove it.  (We would have to delete the polygon that owns
 	// this edge, which is not the intent of this feature.).  We also can't cleanly remove edges that are joining more than two 
@@ -4641,8 +4613,8 @@ void UEditableMesh::TryToRemovePolygonEdge( const FEdgeID EdgeID, bool& bOutWasE
 			const FPolygonID PolygonAID = GetEdgeConnectedPolygon( EdgeID, 0 );
 			const FPolygonID PolygonBID = GetEdgeConnectedPolygon( EdgeID, 1 );
 
-			const TArray<FVertexInstanceID>& PolygonAVertexInstanceIDs = GetMeshDescription()->GetPolygonVertexInstances( PolygonAID );
-			const TArray<FVertexInstanceID>& PolygonBVertexInstanceIDs = GetMeshDescription()->GetPolygonVertexInstances( PolygonBID );
+			TArray<FVertexInstanceID> PolygonAVertexInstanceIDs = GetMeshDescription()->GetPolygonVertexInstances( PolygonAID );
+			TArray<FVertexInstanceID> PolygonBVertexInstanceIDs = GetMeshDescription()->GetPolygonVertexInstances( PolygonBID );
 
 			// If the polygons are in different polygon groups, we can't remove the edge because we can't determine
 			// which polygon group the replacing polygon should belong to
@@ -4778,7 +4750,7 @@ void UEditableMesh::TryToRemoveVertex( const FVertexID VertexID, bool& bOutWasVe
 	EM_ENTER( TEXT( "TryToRemoveVertex: %s" ), *VertexID.ToString() );
 
 	bOutWasVertexRemoved = false;
-	OutNewEdgeID = FEdgeID::Invalid;
+	OutNewEdgeID = INDEX_NONE;
 
 	// We only support removing vertices that are shared by just two edges
 	const int32 ConnectedEdgeCount = GetVertexConnectedEdgeCount( VertexID );
@@ -5194,7 +5166,7 @@ void UEditableMesh::ExtrudePolygons( const TArray<FPolygonID>& PolygonIDs, const
 			NewFrontPolygonVertices.Reset( PolygonVertexIDs.Num() );
 			NewFrontPolygonVertices.SetNum( PolygonVertexIDs.Num(), false );
 
-			const TArray<FVertexInstanceID>& VertexInstanceIDs = GetMeshDescription()->GetPolygonVertexInstances( PolygonID );
+			TArray<FVertexInstanceID> VertexInstanceIDs = GetMeshDescription()->GetPolygonVertexInstances( PolygonID );
 
 			// Map all of the polygon's vertex IDs to their extruded counterparts to create the new polygon perimeter
 			for( int32 PolygonVertexNumber = 0; PolygonVertexNumber < PolygonVertexIDs.Num(); ++PolygonVertexNumber )
@@ -5384,7 +5356,7 @@ void UEditableMesh::ExtendEdges( const TArray<FEdgeID>& EdgeIDs, const bool bWel
 			// which polygons are already connected to the current edge, and use PolygonGroup number from the first
 			// polygon we can find.  If no polygons are connected, then we'll just use first PolygonGroup in the mesh.
 			// We'll also capture texture coordinates from this polygon, so we can apply them to the new polygon vertices.
-			FPolygonID ConnectedPolygonID = FPolygonID::Invalid;
+			FPolygonID ConnectedPolygonID = INDEX_NONE;
 			{
 				const int32 ConnectedPolygonCount = GetEdgeConnectedPolygonCount( OriginalEdgeID );
 				if( ConnectedPolygonCount > 0 )
@@ -5393,8 +5365,8 @@ void UEditableMesh::ExtendEdges( const TArray<FEdgeID>& EdgeIDs, const bool bWel
 				}
 			}
 
-			PolygonToCreate.PolygonGroupID = ( ConnectedPolygonID != FPolygonID::Invalid ) ? GetGroupForPolygon( ConnectedPolygonID ) : GetFirstValidPolygonGroup();
-			check( PolygonToCreate.PolygonGroupID != FPolygonGroupID::Invalid );
+			PolygonToCreate.PolygonGroupID = ( ConnectedPolygonID != INDEX_NONE ) ? GetGroupForPolygon( ConnectedPolygonID ) : GetFirstValidPolygonGroup();
+			check( PolygonToCreate.PolygonGroupID != INDEX_NONE );
 
 			PolygonToCreate.PerimeterVertices.SetNum( 4, false );
 
@@ -5405,7 +5377,7 @@ void UEditableMesh::ExtendEdges( const TArray<FEdgeID>& EdgeIDs, const bool bWel
 			PolygonToCreate.PerimeterVertices[ 3 ].VertexID = ExtendedEdgeVertexIDs[ 1 ];
 
 			// Preserve polygon vertex attributes
-			if( ConnectedPolygonID != FPolygonID::Invalid )
+			if( ConnectedPolygonID != INDEX_NONE )
 			{
 				// Get vertex instance IDs on this polygon corresponding to the edge start/end vertices
 				const FVertexInstanceID EdgeVertexInstanceID0 = GetMeshDescription()->GetVertexInstanceForPolygonVertex( ConnectedPolygonID, OriginalEdgeVertexIDs[ 0 ] );
@@ -5474,7 +5446,7 @@ void UEditableMesh::ExtendVertices( const TArray<FVertexID>& VertexIDs, const bo
 			const FVertexID OriginalVertexID = VertexIDs[ VertexNumber ];
 			const FVertexID NewVertexID = OutNewExtendedVertexIDs[ VertexNumber ];
 
-			FVertexID ClosestVertexID = FVertexID::Invalid;
+			FVertexID ClosestVertexID = INDEX_NONE;
 			if( bOnlyExtendClosestEdge )
 			{
 				// Iterate over the edges connected to this vertex, and figure out which edge is closest to the 
@@ -5510,7 +5482,7 @@ void UEditableMesh::ExtendVertices( const TArray<FVertexID>& VertexIDs, const bo
 					FPolygonToCreate& PolygonToCreate = *new( PolygonsToCreate ) FPolygonToCreate();
 
 					// Figure out which of the connected polygons shares the edge we're going to be using
-					FPolygonID ConnectedPolygonID = FPolygonID::Invalid;
+					FPolygonID ConnectedPolygonID = INDEX_NONE;
 					{
 						static TArray<FPolygonID> ConnectedPolygonIDs;
 						GetVertexConnectedPolygons( OriginalVertexID, /* Out */ ConnectedPolygonIDs );
@@ -5528,7 +5500,7 @@ void UEditableMesh::ExtendVertices( const TArray<FVertexID>& VertexIDs, const bo
 					}
 
 					bool bConnectedPolygonWindsForward = true;
-					if( ConnectedPolygonID != FPolygonID::Invalid )
+					if( ConnectedPolygonID != INDEX_NONE )
 					{
 						const int32 OriginalVertexNumberOnConnectedPolygon = FindPolygonPerimeterVertexNumberForVertex( ConnectedPolygonID, OriginalVertexID );
 						check( OriginalVertexNumberOnConnectedPolygon != INDEX_NONE );
@@ -5545,8 +5517,8 @@ void UEditableMesh::ExtendVertices( const TArray<FVertexID>& VertexIDs, const bo
 						}
 					}
 
-					PolygonToCreate.PolygonGroupID = ( ConnectedPolygonID != FPolygonID::Invalid ) ? GetGroupForPolygon( ConnectedPolygonID ) : GetFirstValidPolygonGroup();
-					check( PolygonToCreate.PolygonGroupID != FPolygonGroupID::Invalid );
+					PolygonToCreate.PolygonGroupID = ( ConnectedPolygonID != INDEX_NONE ) ? GetGroupForPolygon( ConnectedPolygonID ) : GetFirstValidPolygonGroup();
+					check( PolygonToCreate.PolygonGroupID != INDEX_NONE );
 
 					FVertexID ConnectedPolygonVertexIDsForTextureCoordinates[ 3 ];
 					PolygonToCreate.PerimeterVertices.SetNum( 3, false );
@@ -5580,13 +5552,13 @@ void UEditableMesh::ExtendVertices( const TArray<FVertexID>& VertexIDs, const bo
 					}
 
 					// Preserve polygon vertex attributes
-					if( ConnectedPolygonID != FPolygonID::Invalid )
+					if( ConnectedPolygonID != INDEX_NONE )
 					{
 						for( int32 PerimeterVertexNumber = 0; PerimeterVertexNumber < PolygonToCreate.PerimeterVertices.Num(); ++PerimeterVertexNumber )
 						{
 							const FVertexID ConnectedPolygonVertexIDForTextureCoordinates = ConnectedPolygonVertexIDsForTextureCoordinates[ PerimeterVertexNumber ];
 							const FVertexInstanceID VertexInstanceID = GetMeshDescription()->GetVertexInstanceForPolygonVertex( ConnectedPolygonID, ConnectedPolygonVertexIDForTextureCoordinates );
-							check( VertexInstanceID != FVertexInstanceID::Invalid );
+							check( VertexInstanceID != INDEX_NONE );
 
 							BackupAllAttributes( PolygonToCreate.PerimeterVertices[ PerimeterVertexNumber ].PolygonVertexAttributes, GetMeshDescription()->VertexInstanceAttributes(), VertexInstanceID );
 						}
@@ -5660,7 +5632,7 @@ void UEditableMesh::BevelOrInsetPolygons( const TArray<FPolygonID>& PolygonIDs,
 		static TArray<FVertexID> PerimeterVertexIDs;
 		GetPolygonPerimeterVertices( PolygonID, /* Out */ PerimeterVertexIDs );
 
-		const TArray<FVertexInstanceID>& PerimeterVertexInstanceIDs = GetMeshDescription()->GetPolygonVertexInstances( PolygonID );
+		TArray<FVertexInstanceID> PerimeterVertexInstanceIDs = GetMeshDescription()->GetPolygonVertexInstances( PolygonID );
 
 		static TArray<FVertexToCreate> VerticesToCreate;
 		VerticesToCreate.Reset( PerimeterVertexIDs.Num() );
@@ -5722,8 +5694,6 @@ void UEditableMesh::BevelOrInsetPolygons( const TArray<FPolygonID>& PolygonIDs,
 
 			// @todo mesheditor: interpolate all other vertex attributes (maybe handle this with InterpAllAttributes)
 
-			// Determine the barycentric weights of
-			FMeshTriangle Triangle;
 			FVector TriangleVertexWeights;
 
 			const FTriangleID TriangleID = ComputeBarycentricWeightForPointOnPolygon( PolygonID, InsetVertexPosition, TriangleVertexWeights );
@@ -5791,7 +5761,7 @@ void UEditableMesh::BevelOrInsetPolygons( const TArray<FPolygonID>& PolygonIDs,
 				BackupAllAttributes( PolygonToCreate.PerimeterVertices[ 1 ].PolygonVertexAttributes, GetMeshDescription()->VertexInstanceAttributes(), RightOriginalVertexInstanceID );
 
 				// Inset right
-				if( BarycentricWeightsForNewVertices[ RightVertexNumber ].Get<0>() != FTriangleID::Invalid )
+				if( BarycentricWeightsForNewVertices[ RightVertexNumber ].Get<0>() != INDEX_NONE )
 				{
 					InterpAllAttributes( PolygonToCreate.PerimeterVertices[ 2 ].PolygonVertexAttributes,
 										 GetMeshDescription()->VertexInstanceAttributes(),
@@ -5806,7 +5776,7 @@ void UEditableMesh::BevelOrInsetPolygons( const TArray<FPolygonID>& PolygonIDs,
 				}
 
 				// Inset left
-				if( BarycentricWeightsForNewVertices[ LeftVertexNumber ].Get<0>() != FTriangleID::Invalid )
+				if( BarycentricWeightsForNewVertices[ LeftVertexNumber ].Get<0>() != INDEX_NONE )
 				{
 					InterpAllAttributes( PolygonToCreate.PerimeterVertices[ 3 ].PolygonVertexAttributes,
 										 GetMeshDescription()->VertexInstanceAttributes(),
@@ -5841,7 +5811,7 @@ void UEditableMesh::BevelOrInsetPolygons( const TArray<FPolygonID>& PolygonIDs,
 				FVertexAndAttributes& NewPerimeterVertex = PolygonToCreate.PerimeterVertices[ NewVertexNumber ];
 				NewPerimeterVertex.VertexID = NewVertexID;
 
-				if( BarycentricWeightsForNewVertices[ NewVertexNumber ].Get<0>() != FTriangleID::Invalid )
+				if( BarycentricWeightsForNewVertices[ NewVertexNumber ].Get<0>() != INDEX_NONE )
 				{
 					InterpAllAttributes( NewPerimeterVertex.PolygonVertexAttributes,
 										 GetMeshDescription()->VertexInstanceAttributes(),
@@ -5914,6 +5884,7 @@ void UEditableMesh::BevelPolygons( const TArray<FPolygonID>& PolygonIDs, const f
 
 void UEditableMesh::GeneratePolygonTangentsAndNormals( const TArray<FPolygonID>& PolygonIDs )
 {
+	PRAGMA_DISABLE_DEPRECATION_WARNINGS
 	FMeshDescription* MD = GetMeshDescription();
 
 	const TVertexAttributesRef<FVector> VertexPositions = MD->VertexAttributes().GetAttributesRef<FVector>( MeshAttribute::Vertex::Position );
@@ -5930,7 +5901,7 @@ void UEditableMesh::GeneratePolygonTangentsAndNormals( const TArray<FPolygonID>&
 	{
 		// Calculate the center of this polygon
 		FVector Center(ForceInitToZero);
-		const TArray<FVertexInstanceID>& VertexInstanceIDs = MD->GetPolygonVertexInstances( PolygonID );
+		TArray<FVertexInstanceID> VertexInstanceIDs = MD->GetPolygonVertexInstances( PolygonID );
 		for( const FVertexInstanceID VertexInstanceID : VertexInstanceIDs )
 		{
 			Center += VertexPositions[ MD->GetVertexInstanceVertex( VertexInstanceID ) ];
@@ -5975,11 +5946,13 @@ void UEditableMesh::GeneratePolygonTangentsAndNormals( const TArray<FPolygonID>&
 		PolygonBinormals[ PolygonID ] = Binormal.GetSafeNormal();
 		PolygonCenters[ PolygonID ] = Center;
 	}
+	PRAGMA_ENABLE_DEPRECATION_WARNINGS
 }
 
 
 void UEditableMesh::GenerateTangentsAndNormals()
 {
+	PRAGMA_DISABLE_DEPRECATION_WARNINGS
 	const TPolygonAttributesRef<FVector> PolygonNormals = GetMeshDescription()->PolygonAttributes().GetAttributesRef<FVector>( MeshAttribute::Polygon::Normal );
 	const TPolygonAttributesRef<FVector> PolygonTangents = GetMeshDescription()->PolygonAttributes().GetAttributesRef<FVector>( MeshAttribute::Polygon::Tangent );
 	const TPolygonAttributesRef<FVector> PolygonBinormals = GetMeshDescription()->PolygonAttributes().GetAttributesRef<FVector>( MeshAttribute::Polygon::Binormal );
@@ -5991,7 +5964,11 @@ void UEditableMesh::GenerateTangentsAndNormals()
 
 	for( const FPolygonID PolygonID : PolygonsPendingNewTangentBasis )
 	{
-		VertexInstanceIDs.Append( GetMeshDescription()->GetPolygonVertexInstances( PolygonID ) );
+		TArray<FVertexInstanceID> PolyVertInstances = GetMeshDescription()->GetPolygonVertexInstances( PolygonID );
+		for (const FVertexInstanceID PolyVertInstance : PolyVertInstances)
+		{
+			VertexInstanceIDs.Add( PolyVertInstance );
+		}
 	}
 
 	static TArray<FAttributesForVertexInstance> AttributesForVertexInstances;
@@ -6054,6 +6031,7 @@ void UEditableMesh::GenerateTangentsAndNormals()
 	}
 
 	SetVertexInstancesAttributes( AttributesForVertexInstances );
+	PRAGMA_ENABLE_DEPRECATION_WARNINGS
 }
 
 void UEditableMesh::FlipTangentsAndNormals()
@@ -6067,7 +6045,11 @@ void UEditableMesh::FlipTangentsAndNormals()
 
 	for( const FPolygonID PolygonID : PolygonsPendingFlipTangentBasis )
 	{
-		VertexInstanceIDs.Append( GetMeshDescription()->GetPolygonVertexInstances( PolygonID ) );
+		TArray<FVertexInstanceID> PolyVertInstances = GetMeshDescription()->GetPolygonVertexInstances( PolygonID );
+		for (const FVertexInstanceID PolyVertInstance : PolyVertInstances)
+		{
+			VertexInstanceIDs.Add( PolyVertInstance );
+		}
 	}
 
 	static TArray<FAttributesForVertexInstance> AttributesForVertexInstances;
@@ -6096,7 +6078,6 @@ void UEditableMesh::FlipTangentsAndNormals()
 
 void UEditableMesh::SplitPolygonalMesh(const FPlane& InPlane, TArray<FPolygonID>& PolygonIDs1, TArray<FPolygonID>& PolygonIDs2, TArray<FEdgeID>& BoundaryEdges)
 {
-	TPolygonAttributesConstRef<FVector> PolygonCenters = GetMeshDescription()->PolygonAttributes().GetAttributesRef<FVector>(MeshAttribute::Polygon::Center);
 	TVertexAttributesConstRef<FVector> VertexPositions = GetMeshDescription()->VertexAttributes().GetAttributesRef<FVector>(MeshAttribute::Vertex::Position);
 
     // Find potential polygons to split
@@ -6124,7 +6105,12 @@ void UEditableMesh::SplitPolygonalMesh(const FPlane& InPlane, TArray<FPolygonID>
         if (PotentialPolygonsToSplitSet.Contains(PolygonID))
         {
             TSet<FVertexInstanceID> VertexInstanceIDs;
-            VertexInstanceIDs.Append(GetMeshDescription()->GetPolygonVertexInstances(PolygonID));
+			TArray<FVertexInstanceID> PolyVertInstances = GetMeshDescription()->GetPolygonVertexInstances( PolygonID );
+			for (const FVertexInstanceID PolyVertInstance : PolyVertInstances)
+			{
+				VertexInstanceIDs.Add( PolyVertInstance );
+			}
+
             int32 PosNeg = 0;
             TArray<bool> PosNegResults;
             for (const auto& VertexInstanceID : VertexInstanceIDs)
@@ -6417,43 +6403,13 @@ void UEditableMesh::GenerateTangentsForPolygons( const TArray< FPolygonID >& Pol
 
 void UEditableMesh::SetVerticesCornerSharpness( const TArray<FVertexID>& VertexIDs, const TArray<float>& VerticesNewSharpness )
 {
-	check( VertexIDs.Num() == VerticesNewSharpness.Num() );
-
-	static TArray<FAttributesForVertex> AttributesForVertices;
-	AttributesForVertices.Reset();
-
-	for( int32 VertexNumber = 0; VertexNumber < VertexIDs.Num(); ++VertexNumber )
-	{
-		const FVertexID VertexID = VertexIDs[ VertexNumber ];
-
-		AttributesForVertices.Emplace();
-		FAttributesForVertex& AttributesForVertex = AttributesForVertices.Last();
-		AttributesForVertex.VertexID = VertexID;
-		AttributesForVertex.VertexAttributes.Attributes.Emplace( MeshAttribute::Vertex::CornerSharpness, 0, FMeshElementAttributeValue( VerticesNewSharpness[ VertexNumber ] ) );
-	}
-
-	SetVerticesAttributes( AttributesForVertices );
+	return;
 }
 
 
 void UEditableMesh::SetEdgesCreaseSharpness( const TArray<FEdgeID>& EdgeIDs, const TArray<float>& EdgesNewCreaseSharpness )
 {
-	check( EdgeIDs.Num() == EdgesNewCreaseSharpness.Num() );
-
-	static TArray<FAttributesForEdge> AttributesForEdges;
-	AttributesForEdges.Reset();
-
-	for( int32 EdgeNumber = 0; EdgeNumber < EdgeIDs.Num(); ++EdgeNumber )
-	{
-		const FEdgeID EdgeID = EdgeIDs[ EdgeNumber ];
-
-		AttributesForEdges.Emplace();
-		FAttributesForEdge& AttributesForEdge = AttributesForEdges.Last();
-		AttributesForEdge.EdgeID = EdgeID;
-		AttributesForEdge.EdgeAttributes.Attributes.Emplace( MeshAttribute::Edge::CreaseSharpness, 0, FMeshElementAttributeValue( EdgesNewCreaseSharpness[ EdgeNumber ] ) );
-	}
-
-	SetEdgesAttributes( AttributesForEdges );
+	return;
 }
 
 
@@ -6649,7 +6605,7 @@ void UEditableMesh::RemovePolygonPerimeterVertices( const FPolygonID PolygonID,
 		{
 			FVertexAndAttributes& RevertVertexToInsert = RevertInput.VerticesToInsert[ VertexToRemoveIter ];
 			RevertVertexToInsert.VertexInstanceID = Polygon.PerimeterContour.VertexInstanceIDs[ FirstVertexNumberToRemove + VertexToRemoveIter ];
-			RevertVertexToInsert.VertexID = FVertexID::Invalid;
+			RevertVertexToInsert.VertexID = INDEX_NONE;
 		}
 
 		AddUndo( MakeUnique<FInsertPolygonPerimeterVerticesChange>( MoveTemp( RevertInput ) ) );
@@ -6814,7 +6770,7 @@ void UEditableMesh::TriangulatePolygons( const TArray<FPolygonID>& PolygonIDs, T
 						FVertexAndAttributes& PerimeterVertex = PolygonToCreate.PerimeterVertices.Last();
 
 						PerimeterVertex.VertexInstanceID = GetMeshDescription()->GetTriangleVertexInstance( TriangleID, TriangleVertexNumber );
-						PerimeterVertex.VertexID = FVertexID::Invalid;
+						PerimeterVertex.VertexID = INDEX_NONE;
 					}
 				}
 			}
@@ -6900,7 +6856,7 @@ void UEditableMesh::WeldVertices( const TArray<FVertexID>& VertexIDsToWeld, FVer
 {
 	EM_ENTER( TEXT( "WeldVertices: %s" ), *LogHelpers::ArrayToString( VertexIDsToWeld ) );
 
-	OutNewVertexID = FVertexID::Invalid;
+	OutNewVertexID = INDEX_NONE;
 
 	// This function takes a list of perimeter vertices and a list of vertices to be welded as input.
 	// It returns a tuple stating whether the result is valid, and the [first, last) range of vertices to be welded.
@@ -7056,7 +7012,7 @@ void UEditableMesh::WeldVertices( const TArray<FVertexID>& VertexIDsToWeld, FVer
 		static TArray<FEdgeID> PolygonEdgeIDs;
 		GetPolygonPerimeterVertices( ConnectedPolygonID, PolygonVertexIDs );
 		GetPolygonPerimeterEdges( ConnectedPolygonID, PolygonEdgeIDs );
-		const TArray<FVertexInstanceID>& PolygonVertexInstanceIDs = GetMeshDescription()->GetPolygonVertexInstances( ConnectedPolygonID );
+		TArray<FVertexInstanceID> PolygonVertexInstanceIDs = GetMeshDescription()->GetPolygonVertexInstances( ConnectedPolygonID );
 
 		// Get the index range of perimeter vertices to be welded.
 		// This should definitely be valid, as any invalid welded poly will have caused early exit, above.
@@ -7258,7 +7214,7 @@ void UEditableMesh::TessellatePolygons( const TArray<FPolygonID>& PolygonIDs, co
 
 		const int32 PerimeterEdgeCount = GetPolygonPerimeterEdgeCount( PolygonID );
 
-		FVertexID PolygonCenterVertexID = FVertexID::Invalid;
+		FVertexID PolygonCenterVertexID = INDEX_NONE;
 		if( TriangleTessellationMode == ETriangleTessellationMode::ThreeTriangles || PerimeterEdgeCount > 6 )
 		{
 			PolygonCenterVertexID = PolygonCenterVertices[ PolygonWithNewCenterVertexNumber++ ];
@@ -7271,7 +7227,7 @@ void UEditableMesh::TessellatePolygons( const TArray<FPolygonID>& PolygonIDs, co
 			static TArray<FVertexID> PerimeterVertexIDs;
 			GetPolygonPerimeterVertices( PolygonID, /* Out */ PerimeterVertexIDs );
 
-			const TArray<FVertexInstanceID>& PerimeterVertexInstanceIDs = GetMeshDescription()->GetPolygonVertexInstances( PolygonID );
+			TArray<FVertexInstanceID> PerimeterVertexInstanceIDs = GetMeshDescription()->GetPolygonVertexInstances( PolygonID );
 
 			const int32 PerimeterVertexCount = PerimeterEdgeCount;
 			const int32 OriginalPerimeterEdgeCount = PerimeterEdgeCount / 2;
@@ -7323,7 +7279,7 @@ void UEditableMesh::TessellatePolygons( const TArray<FPolygonID>& PolygonIDs, co
 
 							FVector TriangleVertexWeights;
 							const FTriangleID TriangleID = ComputeBarycentricWeightForPointOnPolygon( PolygonID, CenterVertexPosition, /* Out */ TriangleVertexWeights );
-							if( TriangleID != FTriangleID::Invalid )
+							if( TriangleID != INDEX_NONE )
 							{
 								TArrayView<const FVertexInstanceID> TriangleVertexInstanceIDs = GetMeshDescription()->GetTriangleVertexInstances( TriangleID );
 								InterpAllAttributes( VertexAndAttributes.PolygonVertexAttributes, GetMeshDescription()->VertexInstanceAttributes(),
@@ -7343,7 +7299,7 @@ void UEditableMesh::TessellatePolygons( const TArray<FPolygonID>& PolygonIDs, co
 		{
 			if( TriangleTessellationMode == ETriangleTessellationMode::ThreeTriangles )
 			{
-				const TArray<FVertexInstanceID>& PerimeterVertexInstanceIDs = GetMeshDescription()->GetPolygonVertexInstances( PolygonID );
+				TArray<FVertexInstanceID> PerimeterVertexInstanceIDs = GetMeshDescription()->GetPolygonVertexInstances( PolygonID );
 				check( PerimeterVertexInstanceIDs.Num() == 3 );
 
 				// Define the three new triangles for the original tessellated triangle
@@ -7371,7 +7327,7 @@ void UEditableMesh::TessellatePolygons( const TArray<FPolygonID>& PolygonIDs, co
 
 						FVector TriangleVertexWeights;
 						const FTriangleID TriangleID = ComputeBarycentricWeightForPointOnPolygon( PolygonID, CenterVertexPosition, /* Out */ TriangleVertexWeights );
-						if( TriangleID != FTriangleID::Invalid )
+						if( TriangleID != INDEX_NONE )
 						{
 							TArrayView<const FVertexInstanceID> TriangleVertexInstanceIDs = GetMeshDescription()->GetTriangleVertexInstances( TriangleID );
 							InterpAllAttributes( PolygonToCreate.PerimeterVertices[ 2 ].PolygonVertexAttributes, GetMeshDescription()->VertexInstanceAttributes(),
@@ -7386,7 +7342,7 @@ void UEditableMesh::TessellatePolygons( const TArray<FPolygonID>& PolygonIDs, co
 				// the center, connecting the three new vertices that we created between each original edge.  The
 				// other three triangles will go in the corners of the original triangle.
 
-				const TArray<FVertexInstanceID>& PerimeterVertexInstanceIDs = GetMeshDescription()->GetPolygonVertexInstances( PolygonID );
+				TArray<FVertexInstanceID> PerimeterVertexInstanceIDs = GetMeshDescription()->GetPolygonVertexInstances( PolygonID );
 				check( PerimeterVertexInstanceIDs.Num() == 6 );	// We split the triangle's 3 edges earlier, so we must have six edges now
 
 				// Define the new center triangle
@@ -7604,7 +7560,7 @@ void UEditableMesh::QuadrangulatePolygonGroup( const FPolygonGroupID PolygonGrou
 	static TMap<FPolygonID, FAdjacentPolygons> AdjacentPolygonsMap;
 	AdjacentPolygonsMap.Reset();
 
-	FPolygonID PolygonIDToMerge1 = FPolygonID::Invalid;
+	FPolygonID PolygonIDToMerge1 = INDEX_NONE;
 	{
 		float BestScore = TNumericLimits<float>::Max();
 
@@ -7643,15 +7599,15 @@ void UEditableMesh::QuadrangulatePolygonGroup( const FPolygonGroupID PolygonGrou
 								const FPolygonID EdgeConnectedPolygonID = GetEdgeConnectedPolygon( PerimeterEdgeID, EdgeConnectedPolygonIndex );
 								if( EdgeConnectedPolygonID != PolygonID )
 								{
-									return GetPolygonPerimeterEdgeCount( EdgeConnectedPolygonID ) == 3 ? EdgeConnectedPolygonID : FPolygonID::Invalid;
+									return GetPolygonPerimeterEdgeCount( EdgeConnectedPolygonID ) == 3 ? EdgeConnectedPolygonID : INDEX_NONE;
 								}
 							}
 						}
 
-						return FPolygonID::Invalid;
+						return INDEX_NONE;
 					}();
 
-					if( AdjacentPolygonID != FPolygonID::Invalid )
+					if( AdjacentPolygonID != INDEX_NONE )
 					{
 						FAdjacentPolygons* AdjacentPolygons = AdjacentPolygonsMap.Find( PolygonID );
 
@@ -7690,8 +7646,8 @@ void UEditableMesh::QuadrangulatePolygonGroup( const FPolygonGroupID PolygonGrou
 								const int32 PrevAdjacentPerimeterEdgeIndex = ( AdjacentPerimeterEdgeIndex + 2 ) % 3;
 								const int32 NextAdjacentPerimeterEdgeIndex = ( AdjacentPerimeterEdgeIndex + 1 ) % 3;
 
-								const TArray<FVertexInstanceID>& PerimeterVertexInstanceIDs = GetMeshDescription()->GetPolygonVertexInstances( PolygonID );
-								const TArray<FVertexInstanceID>& AdjacentPerimeterVertexInstanceIDs = GetMeshDescription()->GetPolygonVertexInstances( AdjacentPolygonID );
+								TArray<FVertexInstanceID> PerimeterVertexInstanceIDs = GetMeshDescription()->GetPolygonVertexInstances( PolygonID );
+								TArray<FVertexInstanceID> AdjacentPerimeterVertexInstanceIDs = GetMeshDescription()->GetPolygonVertexInstances( AdjacentPolygonID );
 
 								const FVertexInstanceID VertexInstanceStart1 = PerimeterVertexInstanceIDs[ PerimeterEdgeIndex ];
 								const FVertexInstanceID VertexInstanceStart2 = AdjacentPerimeterVertexInstanceIDs[ AdjacentPerimeterEdgeIndex ];
@@ -7769,7 +7725,7 @@ void UEditableMesh::QuadrangulatePolygonGroup( const FPolygonGroupID PolygonGrou
 	}
 
 	// If there were no valid pairs of polys to merge, finish already
-	if( PolygonIDToMerge1 == FPolygonID::Invalid )
+	if( PolygonIDToMerge1 == INDEX_NONE )
 	{
 		return;
 	}
@@ -7892,7 +7848,7 @@ void UEditableMesh::QuadrangulatePolygonGroup( const FPolygonGroupID PolygonGrou
 		// Now look for the next polygon to use: it is the one with the best score from the BoundaryPolygons set.
 
 		float BestScore = TNumericLimits<float>::Max();
-		PolygonIDToMerge1 = FPolygonID::Invalid;
+		PolygonIDToMerge1 = INDEX_NONE;
 		for( const FPolygonID BoundaryPolygon : BoundaryPolygons )
 		{
 			const FAdjacentPolygons& AdjacentPolygons = AdjacentPolygonsMap.FindChecked( BoundaryPolygon );
@@ -7906,7 +7862,7 @@ void UEditableMesh::QuadrangulatePolygonGroup( const FPolygonGroupID PolygonGrou
 		// If there are still no candidates adjacent to the already quadrangulated area, choose the best candidate elsewhere.
 		// This will start a new quadrangulated area, which is grown in the same way as the last.
 
-		if( PolygonIDToMerge1 == FPolygonID::Invalid )
+		if( PolygonIDToMerge1 == INDEX_NONE )
 		{
 			BoundaryPolygons.Reset();
 
@@ -7928,7 +7884,7 @@ void UEditableMesh::QuadrangulatePolygonGroup( const FPolygonGroupID PolygonGrou
 
 		// If there are still no candidates, we've done as much as we can do
 
-		if( PolygonIDToMerge1 == FPolygonID::Invalid )
+		if( PolygonIDToMerge1 == INDEX_NONE )
 		{
 			break;
 		}
diff --git a/Engine/Plugins/Runtime/EditableMesh/Source/EditableMesh/EditableMeshChanges.h b/Engine/Plugins/Runtime/EditableMesh/Source/EditableMesh/EditableMeshChanges.h
--- a/Engine/Plugins/Runtime/EditableMesh/Source/EditableMesh/EditableMeshChanges.h
+++ b/Engine/Plugins/Runtime/EditableMesh/Source/EditableMesh/EditableMeshChanges.h
@@ -647,7 +647,7 @@ struct FInsertPolygonPerimeterVerticesChangeInput
 	
 	/** Default constructor */
 	FInsertPolygonPerimeterVerticesChangeInput()
-		: PolygonID( FPolygonID::Invalid ),
+		: PolygonID( INDEX_NONE ),
 		  InsertBeforeVertexNumber( 0 ),
 		  VerticesToInsert()
 	{
@@ -701,7 +701,7 @@ struct FRemovePolygonPerimeterVerticesChangeInput
 	
 	/** Default constructor */
 	FRemovePolygonPerimeterVerticesChangeInput()
-		: PolygonID( FPolygonID::Invalid ),
+		: PolygonID( INDEX_NONE ),
 		  FirstVertexNumberToRemove( 0 ),
 		  NumVerticesToRemove( 0 ),
 		  bDeleteOrphanedVertexInstances( false )
